MoviePipeline: Added a Cinematic Prestreaming plugin. This experimental plugin allows recording of certain types of asyncronous streaming requests (Nanite and Virtual Texture at this time), then play them back in a Level Sequence at a later point in time, but in advance of when they are actually needed. Resolving may still be visible as the GPU pool may be full and may not be able to hold the request, or it may not have started trying to reach the GPU in time (depending on disk I/O, bandwith, etc.)

This is an experiment meant for improving quality of offline rendered cinematics to avoid having resolving textures being visible in the final frames.

diff --git a/Engine/Plugins/Experimental/CinematicPrestreaming/CinematicPrestreaming.uplugin b/Engine/Plugins/Experimental/CinematicPrestreaming/CinematicPrestreaming.uplugin
new file mode 100644
--- /dev/null
+++ b/Engine/Plugins/Experimental/CinematicPrestreaming/CinematicPrestreaming.uplugin
@@ -0,0 +1,35 @@
+{
+	"FileVersion": 3,
+	"Version": 1,
+	"VersionName": "1.0",
+	"FriendlyName": "Cinematic Prestreaming",
+	"Description": "Adds a way to record certain types of streaming data requests in cinematic cutscenes. The requests can then be played back in advance on the Sequencer timeline to pre-stream data during normal gameplay/rendering.",
+	"Category": "Rendering",
+	"CreatedBy": "Epic Games, Inc.",
+	"CreatedByURL": "",
+	"DocsURL": "",
+	"MarketplaceURL": "",
+	"SupportURL": "",
+	"CanContainContent": true,
+	"IsExperimentalVersion": true,
+	"Installed": false,
+	"Modules": [
+		{
+			"Name": "CinematicPrestreaming",
+			"Type": "Runtime",
+			"LoadingPhase": "Default"
+		},
+		{
+			"Name": "CinematicPrestreamingEditor",
+			"Type": "UncookedOnly",
+			"LoadingPhase": "Default"
+		}
+	],
+	"Plugins":
+	[
+		{
+			"Name": "MovieRenderPipeline",
+			"Enabled": true
+		}
+	]
+}
\ No newline at end of file
diff --git a/Engine/Plugins/Experimental/CinematicPrestreaming/Source/CinematicPrestreaming/CinematicPrestreaming.Build.cs b/Engine/Plugins/Experimental/CinematicPrestreaming/Source/CinematicPrestreaming/CinematicPrestreaming.Build.cs
new file mode 100644
--- /dev/null
+++ b/Engine/Plugins/Experimental/CinematicPrestreaming/Source/CinematicPrestreaming/CinematicPrestreaming.Build.cs
@@ -0,0 +1,31 @@
+// Copyright Epic Games, Inc. All Rights Reserved.
+
+using UnrealBuildTool;
+
+public class CinematicPrestreaming : ModuleRules
+{
+	public CinematicPrestreaming(ReadOnlyTargetRules Target) : base(Target)
+	{
+		ShortName = "CinePrestream";
+
+		PCHUsage = ModuleRules.PCHUsageMode.UseExplicitOrSharedPCHs;
+
+		PublicDependencyModuleNames.AddRange(
+			new string[]
+			{
+				"Core",
+				"MovieScene"
+			}
+		);
+		
+		PrivateDependencyModuleNames.AddRange(
+			new string[]
+			{
+				"CoreUObject",
+				"Engine",
+				"RenderCore",
+				"RHI"
+			}
+		);
+	}
+}
diff --git a/Engine/Plugins/Experimental/CinematicPrestreaming/Source/CinematicPrestreaming/Private/CinePrestreamingData.cpp b/Engine/Plugins/Experimental/CinematicPrestreaming/Source/CinematicPrestreaming/Private/CinePrestreamingData.cpp
new file mode 100644
--- /dev/null
+++ b/Engine/Plugins/Experimental/CinematicPrestreaming/Source/CinematicPrestreaming/Private/CinePrestreamingData.cpp
@@ -0,0 +1,21 @@
+// Copyright Epic Games, Inc. All Rights Reserved.
+
+#include "CinePrestreamingData.h"
+
+void UCinePrestreamingData::GetResourceSizeEx(FResourceSizeEx& CumulativeResourceSize)
+{
+	int32 Size = sizeof(UCinePrestreamingData);
+	Size += Times.GetAllocatedSize();
+	Size += VirtualTextureDatas.GetAllocatedSize();
+	for (FCinePrestreamingVTData const& Data : VirtualTextureDatas)
+	{
+		Size += Data.PageIds.GetAllocatedSize();
+	}
+
+	Size += NaniteDatas.GetAllocatedSize();
+	for (FCinePrestreamingNaniteData const& Data : NaniteDatas)
+	{
+		Size += Data.RequestData.GetAllocatedSize();
+	}
+	CumulativeResourceSize.AddDedicatedSystemMemoryBytes(Size);
+}
diff --git a/Engine/Plugins/Experimental/CinematicPrestreaming/Source/CinematicPrestreaming/Private/CinePrestreamingLog.h b/Engine/Plugins/Experimental/CinematicPrestreaming/Source/CinematicPrestreaming/Private/CinePrestreamingLog.h
new file mode 100644
--- /dev/null
+++ b/Engine/Plugins/Experimental/CinematicPrestreaming/Source/CinematicPrestreaming/Private/CinePrestreamingLog.h
@@ -0,0 +1,7 @@
+// Copyright Epic Games, Inc. All Rights Reserved.
+
+#pragma once
+
+#include "CoreMinimal.h"
+
+DECLARE_LOG_CATEGORY_EXTERN(LogCinePrestreaming, Log, All);
diff --git a/Engine/Plugins/Experimental/CinematicPrestreaming/Source/CinematicPrestreaming/Private/CinePrestreamingModule.cpp b/Engine/Plugins/Experimental/CinematicPrestreaming/Source/CinematicPrestreaming/Private/CinePrestreamingModule.cpp
new file mode 100644
--- /dev/null
+++ b/Engine/Plugins/Experimental/CinematicPrestreaming/Source/CinematicPrestreaming/Private/CinePrestreamingModule.cpp
@@ -0,0 +1,9 @@
+// Copyright Epic Games, Inc. All Rights Reserved.
+
+#include "CinePrestreamingLog.h"
+#include "Modules/ModuleInterface.h"
+#include "Modules/ModuleManager.h"
+
+IMPLEMENT_MODULE(FDefaultModuleImpl, CinePrestreaming);
+
+DEFINE_LOG_CATEGORY(LogCinePrestreaming);
diff --git a/Engine/Plugins/Experimental/CinematicPrestreaming/Source/CinematicPrestreaming/Private/Sections/MovieSceneCinePrestreamingSection.cpp b/Engine/Plugins/Experimental/CinematicPrestreaming/Source/CinematicPrestreaming/Private/Sections/MovieSceneCinePrestreamingSection.cpp
new file mode 100644
--- /dev/null
+++ b/Engine/Plugins/Experimental/CinematicPrestreaming/Source/CinematicPrestreaming/Private/Sections/MovieSceneCinePrestreamingSection.cpp
@@ -0,0 +1,23 @@
+// Copyright Epic Games, Inc. All Rights Reserved.
+
+#include "Sections/MovieSceneCinePrestreamingSection.h"
+#include "TrackInstances/MovieSceneCinePrestreamingTrackInstance.h"
+
+UMovieSceneCinePrestreamingSection::UMovieSceneCinePrestreamingSection(const FObjectInitializer& ObjInit)
+	: Super(ObjInit)
+{
+	EvalOptions.EnableAndSetCompletionMode(EMovieSceneCompletionMode::RestoreState);
+}
+
+void UMovieSceneCinePrestreamingSection::ImportEntityImpl(UMovieSceneEntitySystemLinker* EntityLinker, const FEntityImportParams& Params, FImportedEntity* OutImportedEntity)
+{
+	using namespace UE::MovieScene;
+
+	FMovieSceneTrackInstanceComponent TrackInstance{ decltype(FMovieSceneTrackInstanceComponent::Owner)(this), UMovieSceneCinePrestreamingTrackInstance::StaticClass() };
+
+	OutImportedEntity->AddBuilder(
+		FEntityBuilder()
+		.AddTag(FBuiltInComponentTypes::Get()->Tags.Master)
+		.Add(FBuiltInComponentTypes::Get()->TrackInstance, TrackInstance)
+	);
+}
diff --git a/Engine/Plugins/Experimental/CinematicPrestreaming/Source/CinematicPrestreaming/Private/TrackInstances/MovieSceneCinePrestreamingTrackInstance.cpp b/Engine/Plugins/Experimental/CinematicPrestreaming/Source/CinematicPrestreaming/Private/TrackInstances/MovieSceneCinePrestreamingTrackInstance.cpp
new file mode 100644
--- /dev/null
+++ b/Engine/Plugins/Experimental/CinematicPrestreaming/Source/CinematicPrestreaming/Private/TrackInstances/MovieSceneCinePrestreamingTrackInstance.cpp
@@ -0,0 +1,183 @@
+// Copyright Epic Games, Inc. All Rights Reserved.
+
+#include "TrackInstances/MovieSceneCinePrestreamingTrackInstance.h"
+
+#include "Algo/Sort.h"
+#include "CinePrestreamingData.h"
+#include "CinePrestreamingLog.h"
+#include "Engine/AssetManager.h"
+#include "Engine/StreamableManager.h"
+#include "EngineModule.h"
+#include "EntitySystem/MovieSceneEntitySystemLinker.h"
+#include "EntitySystem/MovieSceneInstanceRegistry.h"
+#include "IMovieScenePlayer.h"
+#include "MovieScene.h"
+#include "MovieSceneTimeHelpers.h"
+#include "RendererInterface.h"
+#include "Sections/MovieSceneCinePrestreamingSection.h"
+#include "Tracks/MovieSceneCinePrestreamingTrack.h"
+
+DECLARE_CYCLE_STAT(TEXT("Cinematic PrestreamingTrack Animate"), MovieSceneEval_CinePrestreamingTrack_Animate, STATGROUP_MovieSceneEval);
+
+static TAutoConsoleVariable<int32> CVarCinematicPreStreamQualityLevel(
+	TEXT("MovieScene.PreStream.QualityLevel"),
+	2,
+	TEXT("We discard prestreaming sections with QualityLevel greater than this."),
+	ECVF_ReadOnly
+);
+
+void UMovieSceneCinePrestreamingTrackInstance::OnInputAdded(const FMovieSceneTrackInstanceInput& InInput)
+{
+	// Trigger async load for the asset associated with this input.
+	const UMovieSceneCinePrestreamingSection* Section = Cast<const UMovieSceneCinePrestreamingSection>(InInput.Section);
+	if (Section != nullptr)
+	{
+		const int32 QualityLevel = Section->GetQualityLevel();
+		if (CVarCinematicPreStreamQualityLevel.GetValueOnGameThread() < QualityLevel)
+		{
+			return;
+		}
+
+		TSoftObjectPtr<UCinePrestreamingData> AssetSoftPtr = Section->GetPrestreamingAsset();
+		if (!AssetSoftPtr.IsNull())
+		{
+			if (PrestreamingAssetMap.Find(InInput) == nullptr)
+			{
+				TWeakObjectPtr<UMovieSceneCinePrestreamingTrackInstance> WeakThis = this;
+				FMovieSceneTrackInstanceInput Input = InInput;
+
+				FStreamableDelegate OnLoadedDelegate;
+				OnLoadedDelegate.BindLambda(
+					[WeakThis, Input, AssetSoftPtr]()
+					{
+						if (UMovieSceneCinePrestreamingTrackInstance* StrongThis = WeakThis.Get())
+						{
+							if (UCinePrestreamingData* AssetPtr = AssetSoftPtr.Get())
+							{
+								StrongThis->PrestreamingAssetMap.Add(Input, AssetPtr);
+								StrongThis->LoadHandleMap.Remove(Input);
+							}
+						}
+					});
+		
+				FStreamableManager& StreamableManager = UAssetManager::GetStreamableManager();
+				TSharedPtr<FStreamableHandle> LoadHandle = StreamableManager.RequestAsyncLoad(AssetSoftPtr.ToSoftObjectPath(), OnLoadedDelegate);
+				LoadHandleMap.Add(InInput, LoadHandle);
+			}
+		}
+	}
+}
+
+void UMovieSceneCinePrestreamingTrackInstance::OnInputRemoved(const FMovieSceneTrackInstanceInput& InInput)
+{
+	const UMovieSceneCinePrestreamingSection* Section = Cast<const UMovieSceneCinePrestreamingSection>(InInput.Section);
+	if (Section != nullptr)
+	{
+		// Remove asset reference should lead to asset unload at some future time.
+		if (PrestreamingAssetMap.Remove(InInput) == 0)
+		{
+			// Didn't remove, so maybe we still have an outstanding load handle.
+			TSharedPtr<FStreamableHandle>* LoadHandlePtr = LoadHandleMap.Find(InInput);
+			if (LoadHandlePtr != nullptr)
+			{
+				FStreamableManager& StreamableManager = UAssetManager::GetStreamableManager();
+				(*LoadHandlePtr)->CancelHandle();
+				LoadHandleMap.Remove(InInput);
+			}
+		}
+	}
+}
+
+void UMovieSceneCinePrestreamingTrackInstance::OnAnimate()
+{
+	SCOPE_CYCLE_COUNTER(MovieSceneEval_CinePrestreamingTrack_Animate);
+
+	using namespace UE::MovieScene;
+
+	const FInstanceRegistry* InstanceRegistry = GetLinker()->GetInstanceRegistry();
+
+	// Gather active prestreaming sections.
+	FCinePrestreamingVTData AccumulatedVirtualTextureData;
+	FCinePrestreamingNaniteData AccumulatedNaniteData;
+
+	for (const FMovieSceneTrackInstanceInput& Input : GetInputs())
+	{
+		const FSequenceInstance& SequenceInstance = InstanceRegistry->GetInstance(Input.InstanceHandle);
+		const FMovieSceneContext& Context = SequenceInstance.GetContext();
+
+		// Get start/current/end time.
+		const TRange<FFrameNumber> SectionRange(Input.Section->GetTrueRange());
+		const FFrameNumber LocalStartTime = SectionRange.HasLowerBound() ? SectionRange.GetLowerBoundValue() : FFrameNumber(-MAX_int32);
+		const FFrameNumber LocalEndTime = SectionRange.HasUpperBound() ? SectionRange.GetUpperBoundValue() : FFrameNumber(MAX_int32);
+		
+		// Get the time relative to the start of the section.
+		const FFrameTime LocalContextTime = Context.GetTime();
+		const FFrameNumber LocalEvalFrame = LocalContextTime.FloorToFrame().Value - LocalStartTime;
+
+		// Apply an offset to start evaluation early (extends into PreRoll so must be less than number of PreRoll frames).
+		const int32 OffsetTime = Input.Section->GetOffsetTime().Get(0).GetFrame().Value;
+		const FFrameRate DisplayRate = Input.Section->GetTypedOuter<UMovieScene>()->GetDisplayRate();
+		const int32 OffsetTimeFrames = FFrameRate::TransformTime(FFrameTime(OffsetTime), DisplayRate, Context.GetFrameRate()).FloorToFrame().Value;
+		const int32 PreRollFrames = Input.Section->GetPreRollFrames();
+		const FFrameNumber StartOffsetFrames = FMath::Clamp(OffsetTimeFrames, 0, PreRollFrames);
+
+		// Early out if we're out of evaluation range.
+		// This can be in the PreRoll, or in the dead space at the end of the shot caused by the OffsetTime.
+		if ((LocalEvalFrame + StartOffsetFrames).Value < 0 || (LocalEvalFrame + StartOffsetFrames) > (LocalEndTime - LocalStartTime))
+		{
+			continue;
+		}
+
+		// Early out if this input is stripped by QualityLevel setting.
+		const UMovieSceneCinePrestreamingSection* Section = Cast<const UMovieSceneCinePrestreamingSection>(Input.Section);
+		const int32 QualityLevel = Section->GetQualityLevel();
+		if (CVarCinematicPreStreamQualityLevel.GetValueOnGameThread() < QualityLevel)
+		{
+			continue;
+		}
+
+		// Get (async loaded) asset data.
+		TObjectPtr<UCinePrestreamingData>* PrestreamingDataPtr = PrestreamingAssetMap.Find(Input);
+		
+		// PreRoll needs to be long enough to cover any async loading before evaluation starts.
+		// In editor this can happen if we are scrubbing the track. But that case is fine.
+		if (!Section->GetPrestreamingAsset().IsNull() && PrestreamingDataPtr == nullptr)
+		{
+			UE_LOG(LogCinePrestreaming, Verbose, TEXT("Prestreaming asset %s didn't load in time. Consider using a longer PreRoll."), *Section->GetPrestreamingAsset().GetAssetName());
+		}
+
+		UCinePrestreamingData const* PrestreamingData = PrestreamingDataPtr != nullptr ? PrestreamingDataPtr->Get() : nullptr;
+		if (PrestreamingData == nullptr || PrestreamingData->Times.Num() == 0)
+		{
+			continue;
+		}
+
+		// Deal with offset due to internal start time of recorded data.
+		const FFrameNumber DataOffsetFrames = PrestreamingData->Times[0];
+
+		// Find the closest (floored) keyframe time.
+		const int32 Index = FMath::Max(0, Algo::UpperBound(PrestreamingData->Times, (LocalEvalFrame + StartOffsetFrames + DataOffsetFrames).Value) - 1);
+
+		// Fetch the data from this keyframe.
+		// todo: Fetch the data from some time range of keyframes?
+		AccumulatedVirtualTextureData.PageIds.Append(PrestreamingData->VirtualTextureDatas[Index].PageIds);
+		AccumulatedNaniteData.RequestData.Append(PrestreamingData->NaniteDatas[Index].RequestData);
+	}
+
+	// Submit virtual texture prestreaming data to the renderer.
+	if (AccumulatedVirtualTextureData.PageIds.Num() > 0)
+	{
+		ENQUEUE_RENDER_COMMAND(RequestTiles)([PageIds = MoveTemp(AccumulatedVirtualTextureData.PageIds)](FRHICommandListImmediate& RHICmdList) mutable
+		{
+			GetRendererModule().RequestVirtualTextureTiles(MoveTemp(PageIds));
+		});
+	}
+
+	if (AccumulatedNaniteData.RequestData.Num() > 0)
+	{
+		ENQUEUE_RENDER_COMMAND(RequestNanitePages)([RequestData = MoveTemp(AccumulatedNaniteData.RequestData)](FRHICommandListImmediate& RHICmdList) mutable
+		{
+			GetRendererModule().RequestNanitePages(MoveTemp(RequestData));
+		});
+	}
+}
diff --git a/Engine/Plugins/Experimental/CinematicPrestreaming/Source/CinematicPrestreaming/Private/TrackInstances/MovieSceneCinePrestreamingTrackInstance.h b/Engine/Plugins/Experimental/CinematicPrestreaming/Source/CinematicPrestreaming/Private/TrackInstances/MovieSceneCinePrestreamingTrackInstance.h
new file mode 100644
--- /dev/null
+++ b/Engine/Plugins/Experimental/CinematicPrestreaming/Source/CinematicPrestreaming/Private/TrackInstances/MovieSceneCinePrestreamingTrackInstance.h
@@ -0,0 +1,31 @@
+// Copyright Epic Games, Inc. All Rights Reserved.
+
+#pragma once
+
+#include "CoreMinimal.h"
+#include "EntitySystem/BuiltInComponentTypes.h"
+#include "EntitySystem/TrackInstance/MovieSceneTrackInstance.h"
+#include "MovieSceneCinePrestreamingTrackInstance.generated.h"
+
+struct FStreamableHandle;
+class UCinePrestreamingData;
+
+UCLASS()
+class UMovieSceneCinePrestreamingTrackInstance : public UMovieSceneTrackInstance
+{
+	GENERATED_BODY()
+
+protected:
+	/** UMovieSceneTrackInstance interface */
+	void OnInputAdded(const FMovieSceneTrackInstanceInput& InInput) override;
+	void OnInputRemoved(const FMovieSceneTrackInstanceInput& InInput) override;
+	void OnAnimate() override;
+
+private:
+	/** Map of loaded asset references. */
+	UPROPERTY(Transient)
+	TMap< FMovieSceneTrackInstanceInput, TObjectPtr<UCinePrestreamingData> > PrestreamingAssetMap;
+
+	/** Map of outstanding async loading handles. */
+	TMap< FMovieSceneTrackInstanceInput, TSharedPtr<FStreamableHandle> > LoadHandleMap;
+};
diff --git a/Engine/Plugins/Experimental/CinematicPrestreaming/Source/CinematicPrestreaming/Private/Tracks/MovieSceneCinePrestreamingTrack.cpp b/Engine/Plugins/Experimental/CinematicPrestreaming/Source/CinematicPrestreaming/Private/Tracks/MovieSceneCinePrestreamingTrack.cpp
new file mode 100644
--- /dev/null
+++ b/Engine/Plugins/Experimental/CinematicPrestreaming/Source/CinematicPrestreaming/Private/Tracks/MovieSceneCinePrestreamingTrack.cpp
@@ -0,0 +1,61 @@
+// Copyright Epic Games, Inc. All Rights Reserved.
+
+#include "Tracks/MovieSceneCinePrestreamingTrack.h"
+#include "Sections/MovieSceneCinePrestreamingSection.h"
+
+#define LOCTEXT_NAMESPACE "MovieSceneCinePrestreamingTrack"
+
+UMovieSceneCinePrestreamingTrack::UMovieSceneCinePrestreamingTrack(const FObjectInitializer& ObjectInitializer)
+	: Super(ObjectInitializer)
+{
+	EvalOptions.bEvaluateInPreroll = EvalOptions.bEvaluateInPostroll = true;
+}
+
+bool UMovieSceneCinePrestreamingTrack::IsEmpty() const
+{
+	return Sections.Num() == 0;
+}
+
+void UMovieSceneCinePrestreamingTrack::AddSection(UMovieSceneSection& Section)
+{
+	Sections.Add(&Section);
+}
+
+void UMovieSceneCinePrestreamingTrack::RemoveSection(UMovieSceneSection& Section)
+{
+	Sections.Remove(&Section);
+}
+
+void UMovieSceneCinePrestreamingTrack::RemoveSectionAt(int32 SectionIndex)
+{
+	Sections.RemoveAt(SectionIndex);
+}
+
+bool UMovieSceneCinePrestreamingTrack::SupportsType(TSubclassOf<UMovieSceneSection> SectionClass) const
+{
+	return SectionClass == UMovieSceneCinePrestreamingSection::StaticClass();
+}
+
+UMovieSceneSection* UMovieSceneCinePrestreamingTrack::CreateNewSection()
+{
+	return NewObject<UMovieSceneCinePrestreamingSection>(this, NAME_None, RF_Transactional);
+}
+
+const TArray<UMovieSceneSection*>& UMovieSceneCinePrestreamingTrack::GetAllSections() const
+{
+	return Sections;
+}
+
+bool UMovieSceneCinePrestreamingTrack::HasSection(const UMovieSceneSection& Section) const
+{
+	return Sections.Contains(&Section);
+}
+
+#if WITH_EDITORONLY_DATA
+FText UMovieSceneCinePrestreamingTrack::GetDefaultDisplayName() const
+{
+	return LOCTEXT("DisplayName", "Cinematic Prestreaming");
+}
+#endif
+
+#undef LOCTEXT_NAMESPACE
diff --git a/Engine/Plugins/Experimental/CinematicPrestreaming/Source/CinematicPrestreaming/Public/CinePrestreamingData.h b/Engine/Plugins/Experimental/CinematicPrestreaming/Source/CinematicPrestreaming/Public/CinePrestreamingData.h
new file mode 100644
--- /dev/null
+++ b/Engine/Plugins/Experimental/CinematicPrestreaming/Source/CinematicPrestreaming/Public/CinePrestreamingData.h
@@ -0,0 +1,51 @@
+// Copyright Epic Games, Inc. All Rights Reserved.
+#pragma once
+
+#include "UObject/Object.h"
+#include "CinePrestreamingData.generated.h"
+
+USTRUCT()
+struct FCinePrestreamingVTData
+{
+	GENERATED_BODY()
+	
+	UPROPERTY(EditAnywhere, Category = "Cinematic Prestreaming")
+	TArray<uint64> PageIds;
+};
+
+USTRUCT()
+struct FCinePrestreamingNaniteData
+{
+	GENERATED_BODY()
+	
+	UPROPERTY(EditAnywhere, Category = "Cinematic Prestreaming")
+	TArray<uint32> RequestData;
+};
+
+UCLASS()
+class CINEMATICPRESTREAMING_API UCinePrestreamingData : public UObject
+{
+	GENERATED_BODY()
+	
+public:
+	UPROPERTY(EditAnywhere, Category = "Cinematic Prestreaming")
+	TArray<FFrameNumber> Times;
+	
+	UPROPERTY(EditAnywhere, Category = "Cinematic Prestreaming")
+	TArray<FCinePrestreamingVTData> VirtualTextureDatas;
+
+	UPROPERTY(EditAnywhere, Category = "Cinematic Prestreaming")
+	TArray<FCinePrestreamingNaniteData> NaniteDatas;
+
+	// UObject Interface
+	void GetResourceSizeEx(FResourceSizeEx& CumulativeResourceSize);
+	// ~UObject Interface
+
+	/** Time that this asset was generated (in UTC). Used to give better context about how up to date an asset is as they are hard to preview. */
+	UPROPERTY(EditAnywhere, Category = "Cinematic Prestreaming")
+	FDateTime RecordedTime;
+
+	/** What resolution was this asset generated at? Recordings are resolution dependent as different mips will be chosen for different resolutions. */
+	UPROPERTY(EditAnywhere, Category = "Cinematic Prestreaming")
+	FIntPoint RecordedResolution;
+};
diff --git a/Engine/Plugins/Experimental/CinematicPrestreaming/Source/CinematicPrestreaming/Public/Sections/MovieSceneCinePrestreamingSection.h b/Engine/Plugins/Experimental/CinematicPrestreaming/Source/CinematicPrestreaming/Public/Sections/MovieSceneCinePrestreamingSection.h
new file mode 100644
--- /dev/null
+++ b/Engine/Plugins/Experimental/CinematicPrestreaming/Source/CinematicPrestreaming/Public/Sections/MovieSceneCinePrestreamingSection.h
@@ -0,0 +1,71 @@
+// Copyright Epic Games, Inc. All Rights Reserved.
+
+#pragma once
+
+#include "EntitySystem/IMovieSceneEntityProvider.h"
+#include "MovieSceneSection.h"
+#include "CinePrestreamingData.h"
+#include "MovieSceneCinePrestreamingSection.generated.h"
+
+class UCinePrestreamingData;
+
+/** Movie Scene Section representing a Prestreaming asset. */
+UCLASS(MinimalAPI)
+class UMovieSceneCinePrestreamingSection
+	: public UMovieSceneSection
+	, public IMovieSceneEntityProvider
+{
+	GENERATED_BODY()
+
+public:
+	UMovieSceneCinePrestreamingSection(const FObjectInitializer& ObjInit);
+
+	/** Get the prestreaming asset soft pointer. Use for async loading. */
+	UFUNCTION(BlueprintPure, Category = "Cinematic Prestreaming")
+	TSoftObjectPtr<UCinePrestreamingData> GetPrestreamingAsset() const { return PrestreamingAsset; }
+
+	UFUNCTION(BlueprintCallable, Category = "Cinematic Prestreaming")
+	void SetPrestreamingAsset(const UCinePrestreamingData* InData) { PrestreamingAsset = TSoftObjectPtr<UCinePrestreamingData>(InData); }
+
+	/** If MovieScene.PreStream.QualityLevel is less than this then discard this section at runtime. */
+	UFUNCTION(BlueprintPure, Category = "Cinematic Prestreaming")
+	int32 GetQualityLevel() const { return QualityLevel; };
+
+	UFUNCTION(BlueprintCallable, Category = "Cinematic Prestreaming")
+	void SetQualityLevel(const int32 InLevel) { QualityLevel = InLevel; }
+
+	UFUNCTION(BlueprintCallable, Category = "Cinematic Prestreaming")
+	void SetStartFrameOffset(const int32 InOffset) { StartFrameOffset = InOffset; }
+
+	/** UMovieSceneSection interface */
+	TOptional<FFrameTime> GetOffsetTime() const override { return TOptional<FFrameTime>(StartFrameOffset); }
+	
+
+protected:
+	/** IMovieSceneEntityProvider interface */
+	void ImportEntityImpl(UMovieSceneEntitySystemLinker* EntityLinker, const FEntityImportParams& Params, FImportedEntity* OutImportedEntity) override;
+
+protected:
+	/** The asset containing cinematic prestreaming data to use for this section. */
+	UPROPERTY(EditAnywhere, BlueprintReadWrite, Category = Prestreaming)
+	TSoftObjectPtr<UCinePrestreamingData> PrestreamingAsset;
+
+	/** Number of frames by which to offset the evaluation. Larger values cause prestreaming to happen earlier. */
+	UPROPERTY(EditAnywhere, BlueprintReadWrite, Category = Prestreaming)
+	int32 StartFrameOffset = 0;
+
+	/** If If MovieScene.PreStream.QualityLevel is less than this then discard this section at runtime. */
+	UPROPERTY(EditAnywhere, BlueprintReadWrite, Category = Prestreaming, AdvancedDisplay)
+	int32 QualityLevel = 0;
+};
+
+
+/** Component data for the cinematic prestreaming system. */
+USTRUCT()
+struct FMovieSceneCinePrestreamingComponentData
+{
+	GENERATED_BODY()
+
+	UPROPERTY()
+	TObjectPtr<const UMovieSceneCinePrestreamingSection> Section = nullptr;
+};
diff --git a/Engine/Plugins/Experimental/CinematicPrestreaming/Source/CinematicPrestreaming/Public/Tracks/MovieSceneCinePrestreamingTrack.h b/Engine/Plugins/Experimental/CinematicPrestreaming/Source/CinematicPrestreaming/Public/Tracks/MovieSceneCinePrestreamingTrack.h
new file mode 100644
--- /dev/null
+++ b/Engine/Plugins/Experimental/CinematicPrestreaming/Source/CinematicPrestreaming/Public/Tracks/MovieSceneCinePrestreamingTrack.h
@@ -0,0 +1,40 @@
+// Copyright Epic Games, Inc. All Rights Reserved.
+
+#pragma once
+
+#include "CoreTypes.h"
+#include "MovieSceneNameableTrack.h"
+#include "MovieSceneCinePrestreamingTrack.generated.h"
+
+/**
+ * A track that controls playback of streaming triggers for some rendering systems.
+ * This is to ensure that data is available before use where by default for systems like
+ * virtual texture and nanite streaming is driven by what is already visible on screen.
+ */
+UCLASS( MinimalAPI )
+class UMovieSceneCinePrestreamingTrack : public UMovieSceneNameableTrack
+{
+	GENERATED_BODY()
+
+public:
+	UMovieSceneCinePrestreamingTrack(const FObjectInitializer& ObjectInitializer);
+
+	/** UMovieSceneTrack interface */
+	bool IsEmpty() const override;
+	void AddSection(UMovieSceneSection& Section) override;
+	void RemoveSection(UMovieSceneSection& Section) override;
+	void RemoveSectionAt(int32 SectionIndex) override;
+	bool SupportsType(TSubclassOf<UMovieSceneSection> SectionClass) const override;
+	UMovieSceneSection* CreateNewSection() override;
+	const TArray<UMovieSceneSection*>& GetAllSections() const override;
+	bool HasSection(const UMovieSceneSection& Section) const override;
+	bool SupportsMultipleRows() const override { return true; }
+
+#if WITH_EDITORONLY_DATA
+	FText GetDefaultDisplayName() const override;
+#endif
+
+private:
+	UPROPERTY()
+	TArray<TObjectPtr<UMovieSceneSection>> Sections;
+};
diff --git a/Engine/Plugins/Experimental/CinematicPrestreaming/Source/CinematicPrestreamingEditor/CinematicPrestreamingEditor.Build.cs b/Engine/Plugins/Experimental/CinematicPrestreaming/Source/CinematicPrestreamingEditor/CinematicPrestreamingEditor.Build.cs
new file mode 100644
--- /dev/null
+++ b/Engine/Plugins/Experimental/CinematicPrestreaming/Source/CinematicPrestreamingEditor/CinematicPrestreamingEditor.Build.cs
@@ -0,0 +1,40 @@
+// Copyright Epic Games, Inc. All Rights Reserved.
+
+using UnrealBuildTool;
+
+public class CinematicPrestreamingEditor : ModuleRules
+{
+	public CinematicPrestreamingEditor(ReadOnlyTargetRules Target) : base(Target)
+	{
+		ShortName = "CinePrestreamEd";
+
+		PCHUsage = ModuleRules.PCHUsageMode.UseExplicitOrSharedPCHs;
+
+		PublicDependencyModuleNames.AddRange(
+			new string[]
+			{
+				"Core",
+				"CinematicPrestreaming",
+				"LevelSequence",
+				"MovieRenderPipelineCore",
+				"MovieRenderPipelineRenderPasses",
+				"MovieRenderPipelineEditor",
+				"MovieScene",
+				"Sequencer",
+				"UnrealEd",
+				"EditorSubsystem"
+			}
+		);
+
+		PrivateDependencyModuleNames.AddRange(
+			new string[]
+			{
+				"CoreUObject",
+				"Engine",
+				"RenderCore",
+				"Slate",
+				"SlateCore",
+			}
+		);
+	}
+}
diff --git a/Engine/Plugins/Experimental/CinematicPrestreaming/Source/CinematicPrestreamingEditor/Private/CinePrestreamingDebugRender.h b/Engine/Plugins/Experimental/CinematicPrestreaming/Source/CinematicPrestreamingEditor/Private/CinePrestreamingDebugRender.h
new file mode 100644
--- /dev/null
+++ b/Engine/Plugins/Experimental/CinematicPrestreaming/Source/CinematicPrestreamingEditor/Private/CinePrestreamingDebugRender.h
@@ -0,0 +1,29 @@
+// Copyright Epic Games, Inc. All Rights Reserved.
+
+#pragma once
+
+#include "MoviePipelineDeferredPasses.h"
+#include "MovieRenderPipelineDataTypes.h"
+#include "CinePrestreamingDebugRender.generated.h"
+
+UCLASS(BlueprintType)
+class UCinePrestreamingDebugRender : public UMoviePipelineDeferredPassBase
+{
+	GENERATED_BODY()
+
+public:
+	UCinePrestreamingDebugRender() : UMoviePipelineDeferredPassBase()
+	{
+		PassIdentifier = FMoviePipelinePassIdentifier("VirtualTexturePendingMips");
+	}
+#if WITH_EDITOR
+	virtual FText GetDisplayText() const override { return NSLOCTEXT("MovieRenderPipeline", "DeferredBasePassSetting_DisplayName_VTPendingMips", "VT Pending Mips (Debug)"); }
+#endif
+	virtual void GetViewShowFlags(FEngineShowFlags& OutShowFlag, EViewModeIndex& OutViewModeIndex) const override
+	{
+		OutShowFlag = FEngineShowFlags(EShowFlagInitMode::ESFIM_Game);
+		OutShowFlag.SetVirtualTexturePendingMips(true);
+		OutViewModeIndex = EViewModeIndex::VMI_VirtualTexturePendingMips;
+	}
+	virtual int32 GetOutputFileSortingOrder() const override { return 2; }
+};
diff --git a/Engine/Plugins/Experimental/CinematicPrestreaming/Source/CinematicPrestreamingEditor/Private/CinePrestreamingEditorModule.cpp b/Engine/Plugins/Experimental/CinematicPrestreaming/Source/CinematicPrestreamingEditor/Private/CinePrestreamingEditorModule.cpp
new file mode 100644
--- /dev/null
+++ b/Engine/Plugins/Experimental/CinematicPrestreaming/Source/CinematicPrestreamingEditor/Private/CinePrestreamingEditorModule.cpp
@@ -0,0 +1,38 @@
+// Copyright Epic Games, Inc. All Rights Reserved.
+
+#include "CinePrestreamingEditorModule.h"
+
+#include "ISequencerModule.h"
+#include "Modules/ModuleManager.h"
+#include "TrackEditors/CinePrestreamingTrackEditor.h"
+
+#define LOCTEXT_NAMESPACE "FCinePrestreamingEditorModule"
+
+void FCinePrestreamingEditorModule::StartupModule()
+{
+	RegisterTrackEditors();
+}
+
+void FCinePrestreamingEditorModule::ShutdownModule()
+{
+	UnregisterTrackEditors();
+}
+
+void FCinePrestreamingEditorModule::RegisterTrackEditors()
+{
+	ISequencerModule& SequencerModule = FModuleManager::LoadModuleChecked<ISequencerModule>("Sequencer");
+	CinePrestreamingCreateEditorHandle = SequencerModule.RegisterTrackEditor(FOnCreateTrackEditor::CreateStatic(&FCinePrestreamingTrackEditor::CreateTrackEditor));
+}
+
+void FCinePrestreamingEditorModule::UnregisterTrackEditors()
+{
+	ISequencerModule* SequencerModule = FModuleManager::GetModulePtr<ISequencerModule>("Sequencer");
+	if (SequencerModule)
+	{
+		SequencerModule->UnRegisterTrackEditor(CinePrestreamingCreateEditorHandle);
+	}
+}
+
+#undef LOCTEXT_NAMESPACE
+	
+IMPLEMENT_MODULE(FCinePrestreamingEditorModule, CinePrestreamingEditor)
\ No newline at end of file
diff --git a/Engine/Plugins/Experimental/CinematicPrestreaming/Source/CinematicPrestreamingEditor/Private/CinePrestreamingEditorSubsystem.cpp b/Engine/Plugins/Experimental/CinematicPrestreaming/Source/CinematicPrestreamingEditor/Private/CinePrestreamingEditorSubsystem.cpp
new file mode 100644
--- /dev/null
+++ b/Engine/Plugins/Experimental/CinematicPrestreaming/Source/CinematicPrestreamingEditor/Private/CinePrestreamingEditorSubsystem.cpp
@@ -0,0 +1,119 @@
+// Copyright Epic Games, Inc. All Rights Reserved.
+#include "CinePrestreamingEditorSubsystem.h"
+#include "CinePrestreamingData.h"
+
+#include "MoviePipelineMasterConfig.h"
+#include "MoviePipelineOutputSetting.h"
+#include "MoviePipelineQueue.h"
+#include "MoviePipelineDeferredPasses.h"
+#include "MoviePipelinePIEExecutor.h"
+#include "CinePrestreamingRecorderSetting.h"
+#include "MoviePipelineQueueSubsystem.h"
+
+#include "ObjectTools.h"
+#include "PackageTools.h"
+#include "PackageHelperFunctions.h"
+#include "PackageTools.h"
+#include "FileHelpers.h"
+#include "Misc/FileHelper.h"
+#include "Misc/PackageName.h"
+#include "AssetRegistryModule.h"
+#include "Editor.h"
+
+
+void UCinePrestreamingEditorSubsystem::GeneratePrestreamingAsset(const FCinePrestreamingGenerateAssetArgs& InArgs)
+{
+	if(!ensureMsgf(!IsRendering(), TEXT("GeneratePrestreamingAsset cannot be called while already generating assets!")))
+	{
+		FFrame::KismetExecutionMessage(TEXT("Asset generation is already in progress."), ELogVerbosity::Error);
+		return;
+	}
+
+	UMoviePipelineQueue* Queue = NewObject<UMoviePipelineQueue>(this);
+	UMoviePipelineExecutorJob* Job = Queue->AllocateNewJob(UMoviePipelineExecutorJob::StaticClass());
+	Job->Sequence = InArgs.Sequence;
+	Job->Map = InArgs.Map;
+	
+	UMoviePipelineOutputSetting* OutputSetting = Cast<UMoviePipelineOutputSetting>(Job->GetConfiguration()->FindOrAddSettingByClass(UMoviePipelineOutputSetting::StaticClass()));
+	OutputSetting->OutputResolution = InArgs.Resolution;
+	
+	// Add the deferred renderer setting
+	Job->GetConfiguration()->FindOrAddSettingByClass(UMoviePipelineDeferredPassBase::StaticClass());
+	
+	// Add our cinematic prestreaming setting to record the requests.
+	UCinePrestreamingRecorderSetting* PrestreamingSetting = Cast<UCinePrestreamingRecorderSetting>(Job->GetConfiguration()->FindOrAddSettingByClass(UCinePrestreamingRecorderSetting::StaticClass()));
+	if (InArgs.OutputDirectoryOverride.Path.Len() > 0)
+	{
+		PrestreamingSetting->PackageDirectory.Path = InArgs.OutputDirectoryOverride.Path;
+	}
+	
+	// No need to add a output like PNG as we don't want to save the rendered images (they aren't
+	// very good looking anyways due to having stuff turned off for perf reasons) and saving them is slow.
+	ActiveAssetArgs = InArgs;
+	
+	ActiveExecutor = NewObject<UMoviePipelinePIEExecutor>(this);
+	ActiveExecutor->OnExecutorFinished().AddUObject(this, &UCinePrestreamingEditorSubsystem::OnBuildPrestreamingComplete);
+	ActiveExecutor->Execute(Queue);
+
+	UMoviePipelineQueueSubsystem* Subsystem = GEditor->GetEditorSubsystem<UMoviePipelineQueueSubsystem>();
+	check(Subsystem);
+
+	// ToDo:
+	// The subsystem only renders the internal queue, but we want to use the subsystem so it appropriately locks the UI out.
+	// OriginalQueue->CopyFrom(Subsystem->GetQueue());
+	// Subsystem->GetQueue()->CopyFrom(Queue);
+	// Subsystem->RenderQueueWithExecutorInstance(LocalExecutor);
+}
+
+void UCinePrestreamingEditorSubsystem::CreatePackagesFromGeneratedData(TArray<FMoviePipelineCinePrestreamingGeneratedData>& InOutData)
+{
+	for (FMoviePipelineCinePrestreamingGeneratedData& Data : InOutData)
+	{
+		const FString FixedAssetName = ObjectTools::SanitizeObjectName(Data.AssetName);
+		FString NewPackageName = FPackageName::GetLongPackagePath(Data.PackagePath) + TEXT("/") + FixedAssetName;
+
+		NewPackageName = UPackageTools::SanitizePackageName(NewPackageName);
+
+		FText OutErrorReason;
+		if (!FPackageName::IsValidLongPackageName(NewPackageName, false, &OutErrorReason))
+		{
+			return;
+		}
+
+		if (UPackage* OldPackage = FindObject<UPackage>(nullptr, *NewPackageName))
+		{
+			FName UniqueName = MakeUniqueObjectName(GetTransientPackage(), UPackage::StaticClass(), "DEAD_CinePrestreamingSettingAsset");
+			OldPackage->Rename(*UniqueName.ToString());
+			OldPackage->SetFlags(RF_Transient);
+		}
+
+		UPackage* NewPackage = CreatePackage(*NewPackageName);
+		NewPackage->AddToRoot();
+
+		// Try to fully load the package (if it already exists) so we can save over it.
+		LoadPackage(NewPackage, *NewPackageName, LOAD_None);
+
+		// Duplicate the data asset into this package
+		UCinePrestreamingData* NewPrestreamingData = Cast<UCinePrestreamingData>(StaticDuplicateObject(Data.StreamingData, NewPackage, FName(*FixedAssetName), RF_NoFlags));
+		NewPrestreamingData->SetFlags(RF_Public | RF_Transactional | RF_Standalone);
+		NewPrestreamingData->MarkPackageDirty();
+
+		Data.StreamingData = NewPrestreamingData;
+
+		// Mark it so it shows up in the Content Browser immediately
+		FAssetRegistryModule::AssetCreated(NewPrestreamingData);
+
+		// If they want to save, ask them to save (and add to version control)
+		TArray<UPackage*> Packages;
+		Packages.Add(NewPrestreamingData->GetOutermost());
+
+		UEditorLoadingAndSavingUtils::SavePackages(Packages, true);
+	}
+}
+
+void UCinePrestreamingEditorSubsystem::OnBuildPrestreamingComplete(UMoviePipelineExecutorBase* InPipelineExecutor, bool bSuccess)
+{
+	OnAssetGenerated.Broadcast(ActiveAssetArgs.GetValue());
+	ActiveAssetArgs.Reset();
+	ActiveExecutor = nullptr;
+}
\ No newline at end of file
diff --git a/Engine/Plugins/Experimental/CinematicPrestreaming/Source/CinematicPrestreamingEditor/Private/CinePrestreamingRecorderSetting.cpp b/Engine/Plugins/Experimental/CinematicPrestreaming/Source/CinematicPrestreamingEditor/Private/CinePrestreamingRecorderSetting.cpp
new file mode 100644
--- /dev/null
+++ b/Engine/Plugins/Experimental/CinematicPrestreaming/Source/CinematicPrestreamingEditor/Private/CinePrestreamingRecorderSetting.cpp
@@ -0,0 +1,360 @@
+// Copyright Epic Games, Inc. All Rights Reserved.
+
+#include "CinePrestreamingRecorderSetting.h"
+
+#include "AssetRegistryModule.h"
+#include "CinePrestreamingData.h"
+#include "EngineModule.h"
+#include "FileHelpers.h"
+#include "Misc/FileHelper.h"
+#include "Misc/PackageName.h"
+#include "MoviePipeline.h"
+#include "MoviePipelineMasterConfig.h"
+#include "MoviePipelineOutputSetting.h"
+#include "MoviePipelineQueue.h"
+#include "MoviePipelineUtils.h"
+#include "MovieScene.h"
+#include "ObjectTools.h"
+#include "PackageHelperFunctions.h"
+#include "PackageTools.h"
+#include "RendererInterface.h"
+#include "RenderingThread.h"
+#include "UObject/Package.h"
+#include "Tracks/MovieSceneCinePrestreamingTrack.h"
+#include "Sections/MovieSceneCinePrestreamingSection.h"
+#include "Kismet/KismetSystemLibrary.h"
+
+#if WITH_EDITOR
+#include "Editor.h"
+#include "CinePrestreamingEditorSubsystem.h"
+#endif
+
+UCinePrestreamingRecorderSetting::UCinePrestreamingRecorderSetting()
+	: PrevActiveShotIndex(-1)
+{
+	PackageDirectory.Path = TEXT("/Game/Cinematics/Prestreaming/{sequence_name}/{shot_name}/");
+
+	ShowFlagsToDisable.Add("ShowFlag.PostProcessing");
+	ShowFlagsToDisable.Add("ShowFlag.Lighting");
+	ShowFlagsToDisable.Add("ShowFlag.Atmosphere");
+	ShowFlagsToDisable.Add("ShowFlag.VolumetricFog");
+	ShowFlagsToDisable.Add("ShowFlag.VolumetricLightmap");
+	ShowFlagsToDisable.Add("ShowFlag.LumenGlobalIllumiination");
+	ShowFlagsToDisable.Add("ShowFlag.LumenReflections");
+	ShowFlagsToDisable.Add("ShowFlag.AmbientOcclusion");
+	ShowFlagsToDisable.Add("ShowFlag.DistanceFieldAO");
+}
+
+void UCinePrestreamingRecorderSetting::SetupForPipelineImpl(UMoviePipeline* InPipeline)
+{
+	// Subscribe to BeginFrame to capture the start of rendering on game thread. We may actually render multiple times to produce a single output Frame.
+	BeginFrameDelegate = FCoreDelegates::OnBeginFrame.AddLambda([this]() { OnBeginFrame_GameThread(); });
+	// Subscribe to EndFrameRT to capture the end of rendering on the render thread. This should be called after all rendering for a single output Frame is complete.
+	EndFrameDelegate = FCoreDelegates::OnEndFrameRT.AddLambda([this]() { OnEndFrame_RenderThread(); });
+
+	SegmentData.Reset();
+	PrevActiveShotIndex = ~0u;
+
+	if (bDisableAdvanceRenderFeatures)
+	{
+		for (const FString& Flag : ShowFlagsToDisable)
+		{
+			FString CombinedStr = Flag + TEXT(" 0");
+			UKismetSystemLibrary::ExecuteConsoleCommand(GetWorld(), CombinedStr, nullptr);
+		}
+	}
+}
+
+void UCinePrestreamingRecorderSetting::TeardownForPipelineImpl(UMoviePipeline* InPipeline)
+{
+	// Force an EndFrame to ensure that all recording is captured. This enqueues a command onto the render
+	// thread, so it needs to come before the flush.
+	OnEndFrame_GameThread();
+
+	// Ensure all render thread work is complete.
+	FlushRenderingCommands();
+
+	// End listening.
+	FCoreDelegates::OnBeginFrame.Remove(BeginFrameDelegate);
+	FCoreDelegates::OnEndFrameRT.Remove(EndFrameDelegate);
+
+	// Save our data asset to disk.
+	CreateAssetsFromData();
+
+	SegmentData.Reset();
+
+	if (bDisableAdvanceRenderFeatures)
+	{
+		for (const FString& Flag : ShowFlagsToDisable)
+		{
+			FString CombinedStr = Flag + TEXT(" 1");
+			UKismetSystemLibrary::ExecuteConsoleCommand(GetWorld(), CombinedStr, nullptr);
+		}
+	}
+}
+	
+void UCinePrestreamingRecorderSetting::OnBeginFrame_GameThread()
+{
+	// Don't capture any data after we finish rendering.
+	if (GetPipeline()->GetPipelineState() != EMovieRenderPipelineState::ProducingFrames)
+	{
+		return;
+	}
+
+	const int32 ActiveShotIndex = GetPipeline()->GetCurrentShotIndex();
+	if (ActiveShotIndex < 0)
+	{
+		return;
+	}
+
+	if (PrevActiveShotIndex != ActiveShotIndex)
+	{
+		// Only allocate the SegmentData array once.
+		// Need to avoid any reallocation after the render thread could be writing it.
+		const int32 ActiveShotNum = GetPipeline()->GetActiveShotList().Num();
+		check(SegmentData.Num() == 0 || SegmentData.Num() == ActiveShotNum);
+		if (SegmentData.Num() == 0)
+		{
+			SegmentData.AddDefaulted(ActiveShotNum);
+		}
+
+		SegmentData[ActiveShotIndex].bValid = true;
+		SegmentData[ActiveShotIndex].OutputState = GetPipeline()->GetOutputState();
+		SegmentData[ActiveShotIndex].InitialShotTick = GetPipeline()->GetActiveShotList()[ActiveShotIndex]->ShotInfo.CurrentTickInMaster;
+
+		PrevActiveShotIndex = ActiveShotIndex;
+	}
+
+	// Be warned that frame time can have repeats and isn't always incrementing, ie: slow-mo will have this
+	// number be the same for many frames, and the start of each shot will jump forward and back to
+	// emulate motion blur, ie: frame 0 -> frame 1 -> frame 0 again. 
+	// So when building data it needs to be additive.
+	UMoviePipelineExecutorShot* CurrentShot = GetPipeline()->GetActiveShotList()[ActiveShotIndex];
+	const FFrameNumber FrameNumber = CurrentShot->ShotInfo.CurrentTickInMaster - SegmentData[ActiveShotIndex].InitialShotTick;
+	
+	// Allocate virtual texture recording page buffer for this frame on the render thread.
+	ENQUEUE_RENDER_COMMAND(SetRecordBuffer)(
+		[ShotIndex = (uint32)PrevActiveShotIndex, FrameNumber](FRHICommandListImmediate& RHICmdList)
+		{
+			const uint64 Handle = ((uint64)ShotIndex << 32) | (uint64)FrameNumber.Value;
+			GetRendererModule().SetVirtualTextureRequestRecordBuffer(Handle);
+			GetRendererModule().SetNaniteRequestRecordBuffer(Handle);
+		});
+}
+
+void UCinePrestreamingRecorderSetting::OnEndFrame_GameThread()
+{
+	ENQUEUE_RENDER_COMMAND(EndFrame)(
+		[this](FRHICommandListImmediate& RHICmdList)
+		{
+			OnEndFrame_RenderThread();
+		});
+}
+
+void UCinePrestreamingRecorderSetting::OnEndFrame_RenderThread()
+{
+	// Finalize request collection for this frame.
+	{
+		// We avoid the need for a critical section by assuming that this is the only code accessing SegmentData during the recording phase.
+		TSet<uint64> PageRequests;
+		uint64 Handle = GetRendererModule().GetVirtualTextureRequestRecordBuffer(PageRequests);
+		if (Handle != ~0ull)
+		{
+			const uint32 ShotIndex = (uint32)(Handle >> 32);
+			const FFrameNumber FrameNumber((int32)(Handle & 0xffffffff));
+			SegmentData[ShotIndex].FrameData.FindOrAdd(FFrameNumber((int32)Handle)).VTRequests.Append(PageRequests);
+		}
+	}
+
+	{
+		TArray<uint32> NaniteRequestData;
+		uint64 Handle = GetRendererModule().GetNaniteRequestRecordBuffer(NaniteRequestData);
+		if (Handle != ~0ull)
+		{
+			const uint32 ShotIndex = (uint32)(Handle >> 32);
+			const FFrameNumber FrameNumber((int32)(Handle & 0xffffffff));
+
+			// TODO: Accumulate the data instead of just setting it?
+			//		 Will require deduplication/encoding logic to be exposed
+			SegmentData[ShotIndex].FrameData.FindOrAdd(FFrameNumber((int32)Handle)).NaniteRequestData = NaniteRequestData;
+		}
+	}
+}
+
+void UCinePrestreamingRecorderSetting::CreateAssetsFromData()
+{
+	UMoviePipelineOutputSetting* OutputSetting = GetPipeline()->GetPipelineMasterConfig()->FindSetting<UMoviePipelineOutputSetting>();
+	FString FileNameFormatString = PackageDirectory.Path;
+
+	const bool bIncludeRenderPass = false;
+	const bool bTestFrameNumber = false;
+ 	UE::MoviePipeline::ValidateOutputFormatString(FileNameFormatString, bIncludeRenderPass, bTestFrameNumber);
+	// Strip any frame number tags so we don't get one asset per frame.
+	UE::MoviePipeline::RemoveFrameNumberFormatStrings(FileNameFormatString, true);
+
+	struct FOutputSegment
+	{
+		FString FilePath;
+		FCollectedData Data;
+		int32 ShotIndex;
+		UMovieScene* MovieScene;
+		TRange<FFrameNumber> Range;
+	};
+
+	TArray<FOutputSegment> FinalSegments;
+
+	for (int32 SegmentIndex = 0; SegmentIndex < SegmentData.Num(); SegmentIndex++)
+	{
+		const FCollectedData& Segment = SegmentData[SegmentIndex];
+		if (!Segment.bValid)
+		{
+			continue;
+		}
+
+		// Generate a filename for this output file.
+		TMap<FString, FString> FormatOverrides;
+		FormatOverrides.Add(TEXT("render_pass"), TEXT("PrestreamingData"));
+		FormatOverrides.Add(TEXT("ext"), TEXT(""));
+
+		const FMoviePipelineCameraCutInfo& ShotInfo = GetPipeline()->GetActiveShotList()[Segment.OutputState.ShotIndex]->ShotInfo;
+
+		// If they don't specify a shot, take it out of the format string to avoid having a folder named "no shot"
+		TSharedPtr<MoviePipeline::FCameraCutSubSectionHierarchyNode> Node = ShotInfo.SubSectionHierarchy;
+		if (!Node->GetParent().IsValid())
+		{
+			FormatOverrides.Add(TEXT("shot_name"), TEXT(""));
+		}
+		
+
+		FMoviePipelineFormatArgs FinalFormatArgs;
+		FString FinalFilePath;
+		GetPipeline()->ResolveFilenameFormatArguments(FileNameFormatString, FormatOverrides, FinalFilePath, FinalFormatArgs, &Segment.OutputState);
+
+		// Remove the final "." ResolveFilenameFormatArgs adds
+		FinalFilePath.LeftChopInline(1);
+		FPaths::NormalizeDirectoryName(FinalFilePath);
+
+		// Look to see if we already have a output file to append this to.
+		FOutputSegment* OutputSegment = nullptr;
+		for (int32 Index = 0; Index < FinalSegments.Num(); Index++)
+		{
+			if (FinalSegments[Index].FilePath == FinalFilePath)
+			{
+				OutputSegment = &FinalSegments[Index];
+				break;
+			}
+		}
+
+		if (!OutputSegment)
+		{
+			FinalSegments.AddDefaulted();
+			OutputSegment = &FinalSegments[FinalSegments.Num() - 1];
+			OutputSegment->FilePath = FinalFilePath;
+			OutputSegment->ShotIndex = Segment.OutputState.ShotIndex;
+
+			// We only use the leafmost moviescene as that's the one that actually gets modified.
+			OutputSegment->MovieScene = Node->MovieScene.Get();
+			OutputSegment->Range = ShotInfo.TotalOutputRangeMaster * ShotInfo.OuterToInnerTransform;
+		}
+
+		// Convert our samples and append them to the existing array
+		OutputSegment->Data.FrameData.Append(Segment.FrameData);
+	}
+
+	TArray<FMoviePipelineCinePrestreamingGeneratedData> GeneratedData;
+	// Now that all of the segments have been merged, write out only the final files.
+	for (FOutputSegment& Segment : FinalSegments)
+	{
+		int32 FrameIndex = 0;
+		TObjectPtr<UCinePrestreamingData> DataTarget = NewObject<UCinePrestreamingData>();
+		DataTarget->RecordedResolution = OutputSetting->OutputResolution;
+		DataTarget->RecordedTime = FDateTime::UtcNow();
+
+		for (TPair< FFrameNumber, FFrameData > const& FrameData : Segment.Data.FrameData)
+		{
+			// Only write out data within the specified range.
+			bool bWriteFrameData = FrameIndex >= StartFrame && (FrameIndex <= EndFrame || EndFrame == 0);
+			// Write a null entry at the end of the data so that during playback we don't keep repeating the last frame.
+			bool bWriteFrameTime = bWriteFrameData || FrameIndex == EndFrame + 1;
+
+			if (bWriteFrameTime)
+			{
+				DataTarget->Times.Add(FrameData.Key);
+			
+				FCinePrestreamingVTData& VTData = DataTarget->VirtualTextureDatas.AddDefaulted_GetRef();
+				FCinePrestreamingNaniteData& NaniteData = DataTarget->NaniteDatas.AddDefaulted_GetRef();
+
+				if (bWriteFrameData)
+				{
+					if (bEnableVirtualTexture)
+					{
+						VTData.PageIds = FrameData.Value.VTRequests.Array();
+						VTData.PageIds.Sort();
+					}
+					if (bEnableNanite)
+					{
+						NaniteData.RequestData = FrameData.Value.NaniteRequestData;
+					}
+				}
+			}
+
+			++FrameIndex;
+		}
+
+		FMoviePipelineCinePrestreamingGeneratedData NewData;
+		NewData.StreamingData = DataTarget;
+
+		ULevelSequence* LevelSequence = Segment.MovieScene->GetTypedOuter<ULevelSequence>();
+		NewData.AssetName = LevelSequence->GetName() + TEXT("_PreStream");
+		NewData.PackagePath = Segment.FilePath;
+		NewData.MovieScene = Segment.MovieScene;
+		NewData.Range = Segment.Range;
+		GeneratedData.Add(NewData);
+	}
+
+	// Packages can only be saved in the editor and need to be saved
+	// before modifying the target sequence, as saving duplicates them
+	// into a package (which will modify GeneratedData to point to the new one)
+#if WITH_EDITOR
+	UCinePrestreamingEditorSubsystem* Subsystem = GEditor->GetEditorSubsystem<UCinePrestreamingEditorSubsystem>();
+	Subsystem->CreatePackagesFromGeneratedData(/*InOut*/GeneratedData);
+#endif
+	if (bModifyTargetSequence)
+	{
+		ModifyTargetSequences(GeneratedData);
+	}
+
+
+	OnGenerateData.Broadcast(GeneratedData);
+}
+
+void UCinePrestreamingRecorderSetting::ModifyTargetSequences(const TArray<FMoviePipelineCinePrestreamingGeneratedData>& InData)
+{
+	for (const FMoviePipelineCinePrestreamingGeneratedData& Data : InData)
+	{
+		const bool bReadOnly = Data.MovieScene->IsReadOnly();
+		Data.MovieScene->SetReadOnly(false);
+
+		// If there's an existing track we'll edit that instead.
+		UMovieSceneCinePrestreamingTrack* PrestreamingTrack = Data.MovieScene->FindMasterTrack<UMovieSceneCinePrestreamingTrack>();
+		if (!PrestreamingTrack)
+		{
+			PrestreamingTrack = Data.MovieScene->AddMasterTrack<UMovieSceneCinePrestreamingTrack>();
+		}
+
+		// Remove any existing sections, as they may point to an old object, or have the wrong length, etc.
+		while (PrestreamingTrack->GetAllSections().Num() > 0)
+		{
+			PrestreamingTrack->RemoveSectionAt(0);
+		}
+
+		// Now create a new section to configure
+		UMovieSceneCinePrestreamingSection* Section = Cast<UMovieSceneCinePrestreamingSection>(PrestreamingTrack->CreateNewSection());
+		PrestreamingTrack->AddSection(*Section);
+
+		Section->SetPrestreamingAsset(Data.StreamingData);
+		TRange<FFrameNumber> SectionRange = Data.Range;
+		Section->SetRange(SectionRange);
+		Data.MovieScene->SetReadOnly(bReadOnly);
+	}
+}
\ No newline at end of file
diff --git a/Engine/Plugins/Experimental/CinematicPrestreaming/Source/CinematicPrestreamingEditor/Private/CinePrestreamingRecorderSetting.h b/Engine/Plugins/Experimental/CinematicPrestreaming/Source/CinematicPrestreamingEditor/Private/CinePrestreamingRecorderSetting.h
new file mode 100644
--- /dev/null
+++ b/Engine/Plugins/Experimental/CinematicPrestreaming/Source/CinematicPrestreamingEditor/Private/CinePrestreamingRecorderSetting.h
@@ -0,0 +1,115 @@
+// Copyright Epic Games, Inc. All Rights Reserved.
+
+#pragma once
+
+#include "MoviePipelineOutputBase.h"
+#include "MovieRenderPipelineDataTypes.h"
+#include "CinePrestreamingRecorderSetting.generated.h"
+
+class UCinePrestreamingData;
+class UMovieScene;
+
+USTRUCT(BlueprintType)
+struct FMoviePipelineCinePrestreamingGeneratedData
+{
+	GENERATED_BODY()
+
+	UPROPERTY(EditAnywhere, BlueprintReadWrite, Category = "Cinematic Prestreaming")
+	UCinePrestreamingData* StreamingData;
+
+	UPROPERTY(EditAnywhere, BlueprintReadWrite, Category = "Cinematic Prestreaming")
+	UMovieScene* MovieScene;
+
+	UPROPERTY(EditAnywhere, BlueprintReadWrite, Category = "Cinematic Prestreaming")
+	FString PackagePath;
+
+	UPROPERTY(EditAnywhere, BlueprintReadWrite, Category = "Cinematic Prestreaming")
+	FString AssetName;
+
+	TRange<FFrameNumber> Range;
+};
+
+DECLARE_DYNAMIC_MULTICAST_DELEGATE_OneParam(FOnCinePrestreamingGenerateData, TArray<FMoviePipelineCinePrestreamingGeneratedData>, GeneratedData);
+
+UCLASS(BlueprintType)
+class UCinePrestreamingRecorderSetting : public UMoviePipelineOutputBase
+{
+	GENERATED_BODY()
+
+public:
+	UCinePrestreamingRecorderSetting();
+
+protected:
+	// UMoviePipelineSetting interface
+	bool IsValidOnShots() const override { return true; }
+	bool IsValidOnMaster() const override { return true; }
+#if WITH_EDITOR
+	FText GetDisplayText() const override { return NSLOCTEXT("MovieRenderPipeline", "CinePrestreamingSettingDisplayName", "Prestreaming Recorder"); }
+#endif
+	void SetupForPipelineImpl(UMoviePipeline* InPipeline) override;
+	void TeardownForPipelineImpl(UMoviePipeline* InPipeline) override;
+	// ~UMoviePipelineSetting Interface
+
+private:
+	void OnBeginFrame_GameThread();
+	void OnEndFrame_GameThread();
+	void OnEndFrame_RenderThread();
+	void ModifyTargetSequences(const TArray<FMoviePipelineCinePrestreamingGeneratedData>& InData);
+	void CreateAssetsFromData();
+
+public:
+	/* 
+	* Specifies which directory the generated assets be placed in. Assumed to be relative to the root folder,
+	* ie it should start with /Game/ etc. Can contain MRQ format strings such as {shot_name} which will
+	* be resolved.
+	*/
+	UPROPERTY(EditAnywhere, BlueprintReadWrite, Category = "Settings",  meta=(ContentDir))
+	FDirectoryPath PackageDirectory;
+	
+	/** Enable capture of virtual texture page requests. */
+	UPROPERTY(EditAnywhere, BlueprintReadWrite, Category = "Settings")
+	bool bEnableVirtualTexture = true;
+	/** Enable capture of nanite requests. */
+	UPROPERTY(EditAnywhere, BlueprintReadWrite, Category = "Settings")
+	bool bEnableNanite = true;
+	/** Automatically add the generated tracks to the target sequence */
+	UPROPERTY(EditAnywhere, BlueprintReadWrite, Category = "Settings")
+	bool bModifyTargetSequence = true;
+	/** Automatically disable advance rendering features that don't influence VT */
+	UPROPERTY(EditAnywhere, BlueprintReadWrite, Category = "Settings")
+	bool bDisableAdvanceRenderFeatures = true;
+	/** First frame to output in each recorded asset. Frames before this are dropped. */
+	UPROPERTY(BlueprintReadWrite, Category = "Settings", meta = (UIMin = "0", ClampMin="0"))
+	int32 StartFrame = 0;
+	/** Last frame to capture in each recorded asset. Frames after this are dropped. (Set to 0 to ignore). */
+	UPROPERTY(BlueprintReadWrite, Category = "Settings", meta = (UIMin = "0", ClampMin = "0"))
+	int32 EndFrame = 0;
+	
+	UPROPERTY(BlueprintAssignable)
+	FOnCinePrestreamingGenerateData OnGenerateData;
+
+private:
+	FDelegateHandle BeginFrameDelegate;
+	FDelegateHandle EndFrameDelegate;
+
+	struct FFrameData
+	{
+		TSet<uint64>	VTRequests;
+		TArray<uint32>	NaniteRequestData;
+	};
+
+	struct FCollectedData
+	{
+		bool bValid = false;
+		/** Map of time to recorded frame requests. */
+		TMap<FFrameNumber, FFrameData> FrameData;
+		/** OutputState required for creating output file names. */
+		FMoviePipelineFrameOutputState OutputState;
+		/** TickInMaster at first tick in frame. */
+		FFrameNumber InitialShotTick;
+	};
+
+	TArray<FCollectedData> SegmentData;
+	int32 PrevActiveShotIndex;
+	TArray<FString> ShowFlagsToDisable;
+};
diff --git a/Engine/Plugins/Experimental/CinematicPrestreaming/Source/CinematicPrestreamingEditor/Private/TrackEditors/CinePrestreamingTrackEditor.cpp b/Engine/Plugins/Experimental/CinematicPrestreaming/Source/CinematicPrestreamingEditor/Private/TrackEditors/CinePrestreamingTrackEditor.cpp
new file mode 100644
--- /dev/null
+++ b/Engine/Plugins/Experimental/CinematicPrestreaming/Source/CinematicPrestreamingEditor/Private/TrackEditors/CinePrestreamingTrackEditor.cpp
@@ -0,0 +1,258 @@
+// Copyright Epic Games, Inc. All Rights Reserved.
+
+#include "TrackEditors/CinePrestreamingTrackEditor.h"
+
+#include "CinePrestreamingData.h"
+#include "EditorStyleSet.h"
+#include "Framework/MultiBox/MultiBoxBuilder.h"
+#include "ISequencerSection.h"
+#include "LevelSequence.h"
+#include "MovieSceneTimeHelpers.h"
+#include "Sections/MovieSceneCinePrestreamingSection.h"
+#include "SequencerSectionPainter.h"
+#include "SequencerUtilities.h"
+#include "Styling/StyleColors.h"
+#include "Tracks/MovieSceneCinePrestreamingTrack.h"
+#include "Widgets/Layout/SBox.h"
+#include "Widgets/SBoxPanel.h"
+#include "Widgets/Text/STextBlock.h"
+#include "Misc/Timespan.h"
+
+#define LOCTEXT_NAMESPACE "CinePrestreamingTrackEditor"
+
+struct FCinePrestreamingSection
+	: public ISequencerSection
+	, public TSharedFromThis<FCinePrestreamingSection>
+{
+	FCinePrestreamingSection(UMovieSceneCinePrestreamingSection* InSection)
+		: WeakSection(InSection)
+	{}
+
+	UMovieSceneSection* GetSectionObject() override
+	{
+		return WeakSection.Get();
+	}
+
+	int32 OnPaintSection(FSequencerSectionPainter& InPainter) const override
+	{
+		return InPainter.PaintSectionBackground();
+	}
+
+	float GetSectionHeight() const override
+	{
+		return 30.f;
+	}
+
+	TSharedRef<SWidget> GenerateSectionWidget() override
+	{
+		return SNew(SBox)
+		.Padding(FMargin(4.f))
+		[
+			SNew(SVerticalBox)
+
+			+ SVerticalBox::Slot()
+			.AutoHeight()
+			[
+				SNew(SHorizontalBox)
+
+				+ SHorizontalBox::Slot()
+				.AutoWidth()
+				[
+					SNew(STextBlock)
+					.Text(this, &FCinePrestreamingSection::GetVisibilityText)
+					.ColorAndOpacity(this, &FCinePrestreamingSection::GetTextColor)
+					.TextStyle(FAppStyle::Get(), "NormalText.Important")
+				]
+
+				+ SHorizontalBox::Slot()
+				.AutoWidth()
+				[
+					SNew(STextBlock)
+					.Text(this, &FCinePrestreamingSection::GetEvalOffsetText)
+				]
+			]
+
+			+ SVerticalBox::Slot()
+			[
+				SNew(STextBlock)
+				.Text(this, &FCinePrestreamingSection::GetLayerBarText)
+				.AutoWrapText(true)
+			]
+		];
+	}
+
+	FText GetVisibilityText() const
+	{
+		UMovieSceneCinePrestreamingSection* Section = WeakSection.Get();
+		if (Section && !Section->GetPrestreamingAsset().IsNull())
+		{
+			return FText::FromString(FString::Printf(TEXT("Asset: %s"), *Section->GetPrestreamingAsset().GetAssetName()));
+		}
+
+		return LOCTEXT("VisibilityText_NoAsset", "No Asset Specified");
+	}
+
+	FText GetEvalOffsetText() const
+	{
+		// UMovieSceneCinePrestreamingSection* Section = WeakSection.Get();
+		// if (Section)
+		// {
+		// 	const FFrameTime OffsetTime = Section->GetOffsetTime().Get(FFrameTime());
+		// 	return FText::Format(LOCTEXT("EvaluationOffsetFormatString", "(Evaluation Offset: {0} Frames)"), OffsetTime.FrameNumber.Value);
+		// }
+
+		return FText();
+	}
+
+	FText GetLayerBarText() const
+	{
+		UMovieSceneCinePrestreamingSection* Section = WeakSection.Get();
+		if (Section)
+		{
+			if (const UCinePrestreamingData* PrestreamingData = Section->GetPrestreamingAsset().Get())
+			{
+				FDateTime CurTime = FDateTime::UtcNow();
+				FTimespan DeltaSinceRecording = (CurTime - PrestreamingData->RecordedTime);
+				FIntPoint Resolution = PrestreamingData->RecordedResolution;
+				FFormatNamedArguments Args;
+				Args.Add("FrameCount", PrestreamingData->Times.Num());
+				Args.Add("ResX", Resolution.X);
+				Args.Add("ResY", Resolution.Y);
+				FString DeltaStrDays = DeltaSinceRecording.ToString(TEXT("%d"));
+				FString DeltaStrHrs = DeltaSinceRecording.ToString(TEXT("%h"));
+				Args.Add("DateDeltaD", FText::FromString(DeltaStrDays));
+				Args.Add("DateDeltaH", FText::FromString(DeltaStrHrs));
+				const FFrameTime OffsetTime = Section->GetOffsetTime().Get(FFrameTime());
+				Args.Add("OffsetFrames", OffsetTime.FrameNumber.Value);
+
+				return FText::Format(LOCTEXT("LayerBarFormatString", "Recorded Duration: {FrameCount} Frames, Resolution: ({ResX}x{ResY}), Recorded: {DateDeltaD} days, {DateDeltaH} hrs ago. Eval Offset: {OffsetFrames} Frames"), Args);
+			}
+		}
+		return FText();
+	}
+
+	FSlateColor GetTextColor() const
+	{
+		return FStyleColors::Foreground;
+	}
+
+private:
+	TWeakObjectPtr<UMovieSceneCinePrestreamingSection> WeakSection;
+};
+
+
+FCinePrestreamingTrackEditor::FCinePrestreamingTrackEditor(TSharedRef<ISequencer> InSequencer)
+	: FMovieSceneTrackEditor(InSequencer)
+{}
+
+TSharedRef<ISequencerTrackEditor> FCinePrestreamingTrackEditor::CreateTrackEditor(TSharedRef<ISequencer> InSequencer)
+{
+	return MakeShared<FCinePrestreamingTrackEditor>(InSequencer);
+}
+
+bool FCinePrestreamingTrackEditor::SupportsSequence(UMovieSceneSequence* InSequence) const
+{
+	return InSequence ? InSequence->GetClass() == ULevelSequence::StaticClass() : false;
+}
+
+bool FCinePrestreamingTrackEditor::SupportsType(TSubclassOf<UMovieSceneTrack> Type) const
+{
+	return Type == UMovieSceneCinePrestreamingTrack::StaticClass();
+}
+
+const FSlateBrush* FCinePrestreamingTrackEditor::GetIconBrush() const
+{
+	return FAppStyle::Get().GetBrush("Sequencer.Tracks.DataLayer");
+}
+
+TSharedRef<ISequencerSection> FCinePrestreamingTrackEditor::MakeSectionInterface(UMovieSceneSection& SectionObject, UMovieSceneTrack& Track, FGuid ObjectBinding)
+{
+	UMovieSceneCinePrestreamingSection* PrestreamingSection = Cast<UMovieSceneCinePrestreamingSection>(&SectionObject);
+	check(SupportsType(SectionObject.GetOuter()->GetClass()) && PrestreamingSection != nullptr);
+
+	return MakeShared<FCinePrestreamingSection>(PrestreamingSection);
+}
+
+void FCinePrestreamingTrackEditor::BuildAddTrackMenu(FMenuBuilder& MenuBuilder)
+{
+	MenuBuilder.AddMenuEntry(
+		LOCTEXT("AddTrack", "Cinematic Prestreaming"),
+		LOCTEXT("AddTrackToolTip", "Adds a new track that can load streaming data used by the renderer to request things in advanced of needing them."),
+		FSlateIcon(FAppStyle::GetAppStyleSetName(), "Sequencer.Tracks.DataLayer"),
+		FUIAction(FExecuteAction::CreateRaw(this, &FCinePrestreamingTrackEditor::HandleAddTrack)));
+}
+
+TSharedPtr<SWidget> FCinePrestreamingTrackEditor::BuildOutlinerEditWidget(const FGuid& ObjectBinding, UMovieSceneTrack* Track, const FBuildEditWidgetParams& Params)
+{
+	return FSequencerUtilities::MakeAddButton(
+		LOCTEXT("AddCinePrestreaming_ButtonLabel", "Cinematic Prestreaming"),
+		FOnGetContent::CreateSP(this, &FCinePrestreamingTrackEditor::BuildAddDataLayerMenu, Track),
+		Params.NodeIsHovered, GetSequencer());
+}
+
+UMovieSceneCinePrestreamingSection* FCinePrestreamingTrackEditor::AddNewSection(UMovieScene* MovieScene, UMovieSceneTrack* PrestreamingTrack)
+{
+	using namespace UE::MovieScene;
+
+	const FScopedTransaction Transaction(LOCTEXT("AddCinePrestreamingSection_Transaction", "Add Cinematic Prestreaming"));
+	PrestreamingTrack->Modify();
+
+	UMovieSceneCinePrestreamingSection* PrestreamingSection = CastChecked<UMovieSceneCinePrestreamingSection>(PrestreamingTrack->CreateNewSection());
+
+	TRange<FFrameNumber> SectionRange = MovieScene->GetPlaybackRange();
+	PrestreamingSection->InitialPlacement(PrestreamingTrack->GetAllSections(), MovieScene->GetPlaybackRange().GetLowerBoundValue(), DiscreteSize(MovieScene->GetPlaybackRange()), true);
+	PrestreamingTrack->AddSection(*PrestreamingSection);
+
+	return PrestreamingSection;
+}
+
+void FCinePrestreamingTrackEditor::HandleAddTrack()
+{
+	UMovieScene* FocusedMovieScene = GetFocusedMovieScene();
+	if (FocusedMovieScene == nullptr || FocusedMovieScene->IsReadOnly())
+	{
+		return;
+	}
+
+	const FScopedTransaction Transaction(LOCTEXT("AddDataLayerTrack_Transaction", "Add Cinematic Prestreaming Track"));
+	FocusedMovieScene->Modify();
+
+	UMovieSceneCinePrestreamingTrack* NewTrack = FocusedMovieScene->AddMasterTrack<UMovieSceneCinePrestreamingTrack>();
+	checkf(NewTrack, TEXT("Failed to create new cinematic prestreaming track."));
+
+	UMovieSceneCinePrestreamingSection* NewSection = AddNewSection(FocusedMovieScene, NewTrack);
+	if (GetSequencer().IsValid())
+	{
+		GetSequencer()->OnAddTrack(NewTrack, FGuid());
+	}
+}
+
+TSharedRef<SWidget> FCinePrestreamingTrackEditor::BuildAddDataLayerMenu(UMovieSceneTrack* PrestreamingTrack)
+{
+	FMenuBuilder MenuBuilder(true, nullptr);
+
+	MenuBuilder.AddMenuEntry(
+		LOCTEXT("AddNewSection", "New Section"),
+		LOCTEXT("AddNewSection_Tip", "Add a new section"),
+		FSlateIcon(),
+		FUIAction(FExecuteAction::CreateSP(
+			this, &FCinePrestreamingTrackEditor::HandleAddNewSection, PrestreamingTrack)));
+
+	return MenuBuilder.MakeWidget();
+}
+
+void FCinePrestreamingTrackEditor::HandleAddNewSection(UMovieSceneTrack* PrestreamingTrack)
+{
+	UMovieScene* FocusedMovieScene = GetFocusedMovieScene();
+	if (FocusedMovieScene)
+	{
+		UMovieSceneCinePrestreamingSection* NewSection = AddNewSection(FocusedMovieScene, PrestreamingTrack);
+
+		GetSequencer()->NotifyMovieSceneDataChanged(EMovieSceneDataChangeType::MovieSceneStructureItemAdded);
+		GetSequencer()->EmptySelection();
+		GetSequencer()->SelectSection(NewSection);
+		GetSequencer()->ThrobSectionSelection();
+	}
+}
+
+#undef LOCTEXT_NAMESPACE
diff --git a/Engine/Plugins/Experimental/CinematicPrestreaming/Source/CinematicPrestreamingEditor/Private/TrackEditors/CinePrestreamingTrackEditor.h b/Engine/Plugins/Experimental/CinematicPrestreaming/Source/CinematicPrestreamingEditor/Private/TrackEditors/CinePrestreamingTrackEditor.h
new file mode 100644
--- /dev/null
+++ b/Engine/Plugins/Experimental/CinematicPrestreaming/Source/CinematicPrestreamingEditor/Private/TrackEditors/CinePrestreamingTrackEditor.h
@@ -0,0 +1,44 @@
+// Copyright Epic Games, Inc. All Rights Reserved.
+
+#pragma once
+
+#include "CoreTypes.h"
+#include "Misc/Guid.h"
+#include "MovieSceneTrackEditor.h"
+
+class FMenuBuilder;
+class UMovieSceneCinePrestreamingSection;
+
+/**
+ * A sequencer track editor for data layer tracks.
+ */
+class FCinePrestreamingTrackEditor
+	: public FMovieSceneTrackEditor
+{
+public:
+	FCinePrestreamingTrackEditor(TSharedRef<ISequencer> InSequencer);
+
+	static TSharedRef<ISequencerTrackEditor> CreateTrackEditor(TSharedRef<ISequencer> OwningSequencer);
+
+protected:
+	/** ISequencerTrackEditor interface */
+	TSharedRef<ISequencerSection> MakeSectionInterface(UMovieSceneSection& SectionObject, UMovieSceneTrack& Track, FGuid ObjectBinding) override;
+	bool SupportsSequence(UMovieSceneSequence* InSequence) const override;
+	bool SupportsType(TSubclassOf<UMovieSceneTrack> Type) const override;
+	const FSlateBrush* GetIconBrush() const override;
+	void BuildAddTrackMenu(FMenuBuilder& MenuBuilder) override;
+	TSharedPtr<SWidget> BuildOutlinerEditWidget(const FGuid& ObjectBinding, UMovieSceneTrack* Track, const FBuildEditWidgetParams& Params) override;
+
+private:
+	/** Adds a new section which spans the length of the owning movie scene with the specified desired state. */
+	UMovieSceneCinePrestreamingSection* AddNewSection(UMovieScene* MovieScene, UMovieSceneTrack* DataLayerTrack /*, EDataLayerState DesiredState*/);
+
+	/** Handles when the add track menu item is activated. */
+	void HandleAddTrack();
+
+	/** Builds the add data layer menu which is displayed on the track. */
+	TSharedRef<SWidget> BuildAddDataLayerMenu(UMovieSceneTrack* DataLayerTrack);
+
+	/** Called to add a new section with a desired state. */
+	void HandleAddNewSection(UMovieSceneTrack* DataLayerTrack/*, EDataLayerState DesiredState*/);
+};
diff --git a/Engine/Plugins/Experimental/CinematicPrestreaming/Source/CinematicPrestreamingEditor/Public/CinePrestreamingEditorModule.h b/Engine/Plugins/Experimental/CinematicPrestreaming/Source/CinematicPrestreamingEditor/Public/CinePrestreamingEditorModule.h
new file mode 100644
--- /dev/null
+++ b/Engine/Plugins/Experimental/CinematicPrestreaming/Source/CinematicPrestreamingEditor/Public/CinePrestreamingEditorModule.h
@@ -0,0 +1,20 @@
+// Copyright Epic Games, Inc. All Rights Reserved.
+
+#pragma once
+
+#include "CoreMinimal.h"
+#include "Modules/ModuleManager.h"
+
+class FCinePrestreamingEditorModule : public IModuleInterface
+{
+public:
+	/** IModuleInterface implementation */
+	void StartupModule() override;
+	void ShutdownModule() override;
+	
+private:
+	void RegisterTrackEditors();
+	void UnregisterTrackEditors();
+	
+	FDelegateHandle CinePrestreamingCreateEditorHandle;
+};
diff --git a/Engine/Plugins/Experimental/CinematicPrestreaming/Source/CinematicPrestreamingEditor/Public/CinePrestreamingEditorSubsystem.h b/Engine/Plugins/Experimental/CinematicPrestreaming/Source/CinematicPrestreamingEditor/Public/CinePrestreamingEditorSubsystem.h
new file mode 100644
--- /dev/null
+++ b/Engine/Plugins/Experimental/CinematicPrestreaming/Source/CinematicPrestreamingEditor/Public/CinePrestreamingEditorSubsystem.h
@@ -0,0 +1,72 @@
+// Copyright Epic Games, Inc. All Rights Reserved.
+
+#pragma once
+
+#include "EditorSubsystem.h"
+#include "CinePrestreamingRecorderSetting.h"
+#include "CinePrestreamingEditorSubsystem.generated.h"
+
+// Forward Declare
+class ULevelSequence;
+class UMoviePipelineExecutorBase;
+
+USTRUCT(BlueprintType)
+struct FCinePrestreamingGenerateAssetArgs
+{
+	GENERATED_BODY()
+	
+	FCinePrestreamingGenerateAssetArgs()
+	{
+		Resolution = FIntPoint(1920, 1080);
+	}
+	
+	UPROPERTY(BlueprintReadWrite, EditAnywhere, Category = "Cinematic Prestreaming|Editor")
+	FDirectoryPath OutputDirectoryOverride;
+	
+	UPROPERTY(BlueprintReadWrite, EditAnywhere, Category = "Cinematic Prestreaming|Editor", meta = (AllowedClasses = "LevelSequence"))
+	FSoftObjectPath Sequence;
+
+	UPROPERTY(BlueprintReadWrite, EditAnywhere, Category = "Cinematic Prestreaming|Editor", meta = (AllowedClasses = "World"))
+	FSoftObjectPath Map;
+
+	UPROPERTY(BlueprintReadWrite, EditAnywhere, Category = "Cinematic Prestreaming|Editor")
+	FIntPoint Resolution;
+};
+
+DECLARE_DYNAMIC_MULTICAST_DELEGATE_OneParam(FOnCinePrestreamingGenerateAsset, FCinePrestreamingGenerateAssetArgs, OriginalGenerationArgs);
+
+
+UCLASS(BlueprintType)
+class CINEMATICPRESTREAMINGEDITOR_API UCinePrestreamingEditorSubsystem : public UEditorSubsystem
+{
+	GENERATED_BODY()
+	
+public:
+	UCinePrestreamingEditorSubsystem() {}
+	
+	UFUNCTION(BlueprintCallable, Category = "Cinematic Prestreaming|Editor")
+	bool IsRendering() const { return ActiveExecutor != nullptr; }
+	
+	UFUNCTION(BlueprintCallable, Category = "Cinematic Prestreaming|Editor")
+	void GeneratePrestreamingAsset(const FCinePrestreamingGenerateAssetArgs& InArgs);
+
+	UFUNCTION(BlueprintCallable, Category = "Cinematic Prestreaming|Editor")
+	void CreatePackagesFromGeneratedData(TArray<FMoviePipelineCinePrestreamingGeneratedData>& InOutData);
+	
+public:
+	UPROPERTY(BlueprintAssignable, Category = "Cinematic Prestreaming|Editor")
+	FOnCinePrestreamingGenerateAsset OnAssetGenerated;
+
+private:
+	UFUNCTION()
+	void OnBuildPrestreamingComplete(UMoviePipelineExecutorBase* InPipelineExecutor, bool bSuccess);
+	
+private:
+	// UPROPERTY(Transient)
+	// UMoviePipelineQueue* OriginalQueue;
+
+	UPROPERTY(Transient)
+	UMoviePipelineExecutorBase* ActiveExecutor;
+
+	TOptional<FCinePrestreamingGenerateAssetArgs> ActiveAssetArgs;
+};
\ No newline at end of file
