MoviePipeline: Fixed some incorrectly initialized struct properties.

diff --git a/Engine/Plugins/Experimental/CinematicPrestreaming/Source/CinematicPrestreamingEditor/Private/CinePrestreamingRecorderSetting.h b/Engine/Plugins/Experimental/CinematicPrestreaming/Source/CinematicPrestreamingEditor/Private/CinePrestreamingRecorderSetting.h
--- a/Engine/Plugins/Experimental/CinematicPrestreaming/Source/CinematicPrestreamingEditor/Private/CinePrestreamingRecorderSetting.h
+++ b/Engine/Plugins/Experimental/CinematicPrestreaming/Source/CinematicPrestreamingEditor/Private/CinePrestreamingRecorderSetting.h
@@ -14,6 +14,11 @@ struct FMoviePipelineCinePrestreamingGeneratedData
 {
 	GENERATED_BODY()
 
+	FMoviePipelineCinePrestreamingGeneratedData()
+		: StreamingData(nullptr)
+		, MovieScene(nullptr)
+	{}
+
 	UPROPERTY(EditAnywhere, BlueprintReadWrite, Category = "Cinematic Prestreaming")
 	TObjectPtr<UCinePrestreamingData> StreamingData;
 
diff --git a/Engine/Plugins/Experimental/CinematicPrestreaming/Source/CinematicPrestreamingEditor/Public/CinePrestreamingEditorSubsystem.h b/Engine/Plugins/Experimental/CinematicPrestreaming/Source/CinematicPrestreamingEditor/Public/CinePrestreamingEditorSubsystem.h
--- a/Engine/Plugins/Experimental/CinematicPrestreaming/Source/CinematicPrestreamingEditor/Public/CinePrestreamingEditorSubsystem.h
+++ b/Engine/Plugins/Experimental/CinematicPrestreaming/Source/CinematicPrestreamingEditor/Public/CinePrestreamingEditorSubsystem.h
@@ -23,10 +23,10 @@ struct FCinePrestreamingGenerateAssetArgs
 	UPROPERTY(BlueprintReadWrite, EditAnywhere, Category = "Cinematic Prestreaming|Editor")
 	FDirectoryPath OutputDirectoryOverride;
 	
-	UPROPERTY(BlueprintReadWrite, EditAnywhere, Category = "Cinematic Prestreaming|Editor", meta = (AllowedClasses = "LevelSequence"))
+	UPROPERTY(BlueprintReadWrite, EditAnywhere, Category = "Cinematic Prestreaming|Editor", meta = (AllowedClasses = "/Script/LevelSequence.LevelSequence"))
 	FSoftObjectPath Sequence;
 
-	UPROPERTY(BlueprintReadWrite, EditAnywhere, Category = "Cinematic Prestreaming|Editor", meta = (AllowedClasses = "World"))
+	UPROPERTY(BlueprintReadWrite, EditAnywhere, Category = "Cinematic Prestreaming|Editor", meta = (AllowedClasses = "/Script/Engine.World"))
 	FSoftObjectPath Map;
 
 	UPROPERTY(BlueprintReadWrite, EditAnywhere, Category = "Cinematic Prestreaming|Editor")
