MoviePipeline: Cleaned up the Cinematic Prestreaming recorder setting a little to help communicate what the feature does to users to reduce support load later when output renders are changed unexpectedly.

diff --git a/Engine/Plugins/Experimental/CinematicPrestreaming/Source/CinematicPrestreamingEditor/Private/CinePrestreamingRecorderSetting.cpp b/Engine/Plugins/Experimental/CinematicPrestreaming/Source/CinematicPrestreamingEditor/Private/CinePrestreamingRecorderSetting.cpp
--- a/Engine/Plugins/Experimental/CinematicPrestreaming/Source/CinematicPrestreamingEditor/Private/CinePrestreamingRecorderSetting.cpp
+++ b/Engine/Plugins/Experimental/CinematicPrestreaming/Source/CinematicPrestreamingEditor/Private/CinePrestreamingRecorderSetting.cpp
@@ -286,12 +286,12 @@ void UCinePrestreamingRecorderSetting::CreateAssetsFromData()
 
 				if (bWriteFrameData)
 				{
-					if (bEnableVirtualTexture)
+					if (bVirtualTextures)
 					{
 						VTData.PageIds = FrameData.Value.VTRequests.Array();
 						VTData.PageIds.Sort();
 					}
-					if (bEnableNanite)
+					if (bNanite)
 					{
 						NaniteData.RequestData = FrameData.Value.NaniteRequestData;
 					}
@@ -357,4 +357,19 @@ void UCinePrestreamingRecorderSetting::ModifyTargetSequences(const TArray<FMovie
 		Section->SetRange(SectionRange);
 		Data.MovieScene->SetReadOnly(bReadOnly);
 	}
-}
\ No newline at end of file
+}
+
+
+#if WITH_EDITOR
+
+FText UCinePrestreamingRecorderSetting::GetFooterText(UMoviePipelineExecutorJob* InJob) const
+{
+	if (bDisableAdvanceRenderFeatures)
+	{
+		return NSLOCTEXT("MovieRenderPipeline", "CinePrestreamingSetting_FooterText", "Enabling the bDisableAdvanceRenderFeatures setting will disable rendering features not required to generate the Cinematic Prestreaming data. This makes renders significantly faster but results in a final image that is not useful.");
+	}
+
+	return FText();
+}
+
+#endif
\ No newline at end of file
diff --git a/Engine/Plugins/Experimental/CinematicPrestreaming/Source/CinematicPrestreamingEditor/Private/CinePrestreamingRecorderSetting.h b/Engine/Plugins/Experimental/CinematicPrestreaming/Source/CinematicPrestreamingEditor/Private/CinePrestreamingRecorderSetting.h
deleted file mode 100644
diff --git a/Engine/Plugins/Experimental/CinematicPrestreaming/Source/CinematicPrestreamingEditor/Public/CinePrestreamingRecorderSetting.h b/Engine/Plugins/Experimental/CinematicPrestreaming/Source/CinematicPrestreamingEditor/Public/CinePrestreamingRecorderSetting.h
new file mode 100644
--- /dev/null
+++ b/Engine/Plugins/Experimental/CinematicPrestreaming/Source/CinematicPrestreamingEditor/Public/CinePrestreamingRecorderSetting.h
@@ -0,0 +1,122 @@
+// Copyright Epic Games, Inc. All Rights Reserved.
+
+#pragma once
+
+#include "MoviePipelineOutputBase.h"
+#include "MovieRenderPipelineDataTypes.h"
+#include "CinePrestreamingRecorderSetting.generated.h"
+
+class UCinePrestreamingData;
+class UMovieScene;
+
+USTRUCT(BlueprintType)
+struct FMoviePipelineCinePrestreamingGeneratedData
+{
+	GENERATED_BODY()
+
+	FMoviePipelineCinePrestreamingGeneratedData()
+		: StreamingData(nullptr)
+		, MovieScene(nullptr)
+	{}
+
+	UPROPERTY(EditAnywhere, BlueprintReadWrite, Category = "Cinematic Prestreaming")
+	TObjectPtr<UCinePrestreamingData> StreamingData;
+
+	UPROPERTY(EditAnywhere, BlueprintReadWrite, Category = "Cinematic Prestreaming")
+	TObjectPtr<UMovieScene> MovieScene;
+
+	UPROPERTY(EditAnywhere, BlueprintReadWrite, Category = "Cinematic Prestreaming")
+	FString PackagePath;
+
+	UPROPERTY(EditAnywhere, BlueprintReadWrite, Category = "Cinematic Prestreaming")
+	FString AssetName;
+
+	TRange<FFrameNumber> Range;
+};
+
+DECLARE_DYNAMIC_MULTICAST_DELEGATE_OneParam(FOnCinePrestreamingGenerateData, TArray<FMoviePipelineCinePrestreamingGeneratedData>, GeneratedData);
+
+UCLASS(BlueprintType)
+class CINEMATICPRESTREAMINGEDITOR_API UCinePrestreamingRecorderSetting : public UMoviePipelineOutputBase
+{
+	GENERATED_BODY()
+
+public:
+	UCinePrestreamingRecorderSetting();
+
+protected:
+	// UMoviePipelineSetting interface
+	bool IsValidOnShots() const override { return true; }
+	bool IsValidOnMaster() const override { return true; }
+#if WITH_EDITOR
+	FText GetDisplayText() const override { return NSLOCTEXT("MovieRenderPipeline", "CinePrestreamingSettingDisplayName", "Prestreaming Recorder"); }
+	FText GetFooterText(UMoviePipelineExecutorJob* InJob) const;
+
+#endif
+	void SetupForPipelineImpl(UMoviePipeline* InPipeline) override;
+	void TeardownForPipelineImpl(UMoviePipeline* InPipeline) override;
+	// ~UMoviePipelineSetting Interface
+
+private:
+	void OnBeginFrame_GameThread();
+	void OnEndFrame_GameThread();
+	void OnEndFrame_RenderThread();
+	void ModifyTargetSequences(const TArray<FMoviePipelineCinePrestreamingGeneratedData>& InData);
+	void CreateAssetsFromData();
+
+public:
+	/* 
+	* Specifies which directory the generated assets be placed in. Assumed to be relative to the root folder,
+	* ie it should start with /Game/ etc. Can contain MRQ format strings such as {shot_name} which will
+	* be resolved.
+	*/
+	UPROPERTY(EditAnywhere, BlueprintReadWrite, Category = "Settings",  meta=(ContentDir))
+	FDirectoryPath PackageDirectory;
+	
+	/** Enable capture of virtual texture page requests. */
+	UPROPERTY(EditAnywhere, BlueprintReadWrite,Category = "Record")
+	bool bVirtualTextures = true;
+	/** Enable capture of nanite requests. */
+	UPROPERTY(EditAnywhere, BlueprintReadWrite, Category = "Record")
+	bool bNanite = true;
+	/** Automatically add the generated tracks to the target sequence */
+	UPROPERTY(EditAnywhere, BlueprintReadWrite, Category = "Post Process")
+	bool bModifyTargetSequence = true;
+	/** By default we disable any render features that aren't needed to generate the prestreaming data. This makes renders significantly faster, but also means that the final image generated by the renderer is unusable. */
+	UPROPERTY(EditAnywhere, BlueprintReadWrite, Category = "Optimizations")
+	bool bDisableAdvanceRenderFeatures = true;
+	/** First frame to output in each recorded asset. Frames before this are dropped. */
+	UPROPERTY(BlueprintReadWrite, Category = "Settings", meta = (UIMin = "0", ClampMin="0"))
+	int32 StartFrame = 0;
+	/** Last frame to capture in each recorded asset. Frames after this are dropped. (Set to 0 to ignore). */
+	UPROPERTY(BlueprintReadWrite, Category = "Settings", meta = (UIMin = "0", ClampMin = "0"))
+	int32 EndFrame = 0;
+	
+	UPROPERTY(BlueprintAssignable)
+	FOnCinePrestreamingGenerateData OnGenerateData;
+
+private:
+	FDelegateHandle BeginFrameDelegate;
+	FDelegateHandle EndFrameDelegate;
+
+	struct FFrameData
+	{
+		TSet<uint64>	VTRequests;
+		TArray<uint32>	NaniteRequestData;
+	};
+
+	struct FCollectedData
+	{
+		bool bValid = false;
+		/** Map of time to recorded frame requests. */
+		TMap<FFrameNumber, FFrameData> FrameData;
+		/** OutputState required for creating output file names. */
+		FMoviePipelineFrameOutputState OutputState;
+		/** TickInMaster at first tick in frame. */
+		FFrameNumber InitialShotTick;
+	};
+
+	TArray<FCollectedData> SegmentData;
+	int32 PrevActiveShotIndex;
+	TArray<FString> ShowFlagsToDisable;
+};
