Cinematic Prestreaming: Rename from CinePrestreaming to CinematicPrestreaming to match uplugin

diff --git a/Engine/Plugins/Experimental/CinematicPrestreaming/Source/CinematicPrestreaming/Private/CinePrestreamingModule.cpp b/Engine/Plugins/Experimental/CinematicPrestreaming/Source/CinematicPrestreaming/Private/CinePrestreamingModule.cpp
--- a/Engine/Plugins/Experimental/CinematicPrestreaming/Source/CinematicPrestreaming/Private/CinePrestreamingModule.cpp
+++ b/Engine/Plugins/Experimental/CinematicPrestreaming/Source/CinematicPrestreaming/Private/CinePrestreamingModule.cpp
@@ -3,6 +3,6 @@
 #include "CinePrestreamingLog.h"
 #include "Modules/ModuleManager.h"
 
-IMPLEMENT_MODULE(FDefaultModuleImpl, CinePrestreaming);
+IMPLEMENT_MODULE(FDefaultModuleImpl, CinematicPrestreaming);
 
 DEFINE_LOG_CATEGORY(LogCinePrestreaming);
diff --git a/Engine/Plugins/Experimental/CinematicPrestreaming/Source/CinematicPrestreamingEditor/Private/CinePrestreamingEditorModule.cpp b/Engine/Plugins/Experimental/CinematicPrestreaming/Source/CinematicPrestreamingEditor/Private/CinePrestreamingEditorModule.cpp
--- a/Engine/Plugins/Experimental/CinematicPrestreaming/Source/CinematicPrestreamingEditor/Private/CinePrestreamingEditorModule.cpp
+++ b/Engine/Plugins/Experimental/CinematicPrestreaming/Source/CinematicPrestreamingEditor/Private/CinePrestreamingEditorModule.cpp
@@ -35,4 +35,4 @@ void FCinePrestreamingEditorModule::UnregisterTrackEditors()
 
 #undef LOCTEXT_NAMESPACE
 	
-IMPLEMENT_MODULE(FCinePrestreamingEditorModule, CinePrestreamingEditor)
\ No newline at end of file
+IMPLEMENT_MODULE(FCinePrestreamingEditorModule, CinematicPrestreamingEditor)
\ No newline at end of file
