Cine Prestreaming: Added DisplayName for the track

#rnx

diff --git a/Engine/Plugins/Experimental/CinematicPrestreaming/Source/CinematicPrestreamingEditor/Private/TrackEditors/CinePrestreamingTrackEditor.cpp b/Engine/Plugins/Experimental/CinematicPrestreaming/Source/CinematicPrestreamingEditor/Private/TrackEditors/CinePrestreamingTrackEditor.cpp
--- a/Engine/Plugins/Experimental/CinematicPrestreaming/Source/CinematicPrestreamingEditor/Private/TrackEditors/CinePrestreamingTrackEditor.cpp
+++ b/Engine/Plugins/Experimental/CinematicPrestreaming/Source/CinematicPrestreamingEditor/Private/TrackEditors/CinePrestreamingTrackEditor.cpp
@@ -148,6 +148,11 @@ TSharedRef<ISequencerTrackEditor> FCinePrestreamingTrackEditor::CreateTrackEdito
 	return MakeShared<FCinePrestreamingTrackEditor>(InSequencer);
 }
 
+FText FCinePrestreamingTrackEditor::GetDisplayName() const
+{
+	return LOCTEXT("CinePrestreamingTrackEditor_DisplayName", "Cinematic Prestreaming");
+}
+
 bool FCinePrestreamingTrackEditor::SupportsSequence(UMovieSceneSequence* InSequence) const
 {
 	return InSequence ? InSequence->GetClass() == ULevelSequence::StaticClass() : false;
diff --git a/Engine/Plugins/Experimental/CinematicPrestreaming/Source/CinematicPrestreamingEditor/Private/TrackEditors/CinePrestreamingTrackEditor.h b/Engine/Plugins/Experimental/CinematicPrestreaming/Source/CinematicPrestreamingEditor/Private/TrackEditors/CinePrestreamingTrackEditor.h
--- a/Engine/Plugins/Experimental/CinematicPrestreaming/Source/CinematicPrestreamingEditor/Private/TrackEditors/CinePrestreamingTrackEditor.h
+++ b/Engine/Plugins/Experimental/CinematicPrestreaming/Source/CinematicPrestreamingEditor/Private/TrackEditors/CinePrestreamingTrackEditor.h
@@ -22,6 +22,7 @@ public:
 
 protected:
 	/** ISequencerTrackEditor interface */
+	virtual FText GetDisplayName() const override;
 	TSharedRef<ISequencerSection> MakeSectionInterface(UMovieSceneSection& SectionObject, UMovieSceneTrack& Track, FGuid ObjectBinding) override;
 	bool SupportsSequence(UMovieSceneSequence* InSequence) const override;
 	bool SupportsType(TSubclassOf<UMovieSceneTrack> Type) const override;
