Cinematic Prestream: Fix VT debug mode to show pending mips.
Also provide option to disable the debug view, or show stack count.

#Synced-CL 50433575

diff --git a/Engine/Plugins/Experimental/CinematicPrestreaming/Source/CinematicPrestreamingEditor/Private/CinePrestreamingDebugRender.cpp b/Engine/Plugins/Experimental/CinematicPrestreaming/Source/CinematicPrestreamingEditor/Private/CinePrestreamingDebugRender.cpp
new file mode 100644
--- /dev/null
+++ b/Engine/Plugins/Experimental/CinematicPrestreaming/Source/CinematicPrestreamingEditor/Private/CinePrestreamingDebugRender.cpp
@@ -0,0 +1,44 @@
+// Copyright Epic Games, Inc. All Rights Reserved.
+
+#include "CinePrestreamingDebugRender.h"
+
+#include "VT/VirtualTextureVisualizationData.h"
+
+UCinePrestreamingDebugRender::UCinePrestreamingDebugRender() 
+	: UMoviePipelineDeferredPassBase()
+	, VirtualTextureDebugMode(EVirtualTextureVisualizationMode::PendingMips)
+	, PreviousMode(EVirtualTextureVisualizationMode::None)
+{
+	PassIdentifier = FMoviePipelinePassIdentifier("VirtualTexturePendingMips");
+	AdditionalPostProcessMaterials.Reset();
+}
+
+#if WITH_EDITOR
+FText UCinePrestreamingDebugRender::GetDisplayText() const
+{
+	return NSLOCTEXT("MovieRenderPipeline", "DeferredBasePassSetting_DisplayName_VTPendingMips", "Deferred Rendering (VT Debug)");
+}
+#endif
+
+void UCinePrestreamingDebugRender::SetupImpl(const MoviePipeline::FMoviePipelineRenderPassInitSettings& InPassInitSettings)
+{
+	Super::SetupImpl(InPassInitSettings);
+
+	PreviousMode = GetVirtualTextureVisualizationData().SetModeUsingConsoleCommand(VirtualTextureDebugMode);
+}
+
+void UCinePrestreamingDebugRender::TeardownImpl()
+{
+	Super::TeardownImpl();
+
+	GetVirtualTextureVisualizationData().SetModeUsingConsoleCommand(PreviousMode);
+}
+
+void UCinePrestreamingDebugRender::GetViewShowFlags(FEngineShowFlags& OutShowFlag, EViewModeIndex& OutViewModeIndex) const
+{
+	const bool bVisualizeVirtualTexture = VirtualTextureDebugMode != EVirtualTextureVisualizationMode::None;
+
+	OutShowFlag = FEngineShowFlags(EShowFlagInitMode::ESFIM_Game);
+	OutShowFlag.SetVisualizeVirtualTexture(bVisualizeVirtualTexture);
+	OutViewModeIndex = bVisualizeVirtualTexture ? EViewModeIndex::VMI_VisualizeVirtualTexture : EViewModeIndex::VMI_Lit;
+}
diff --git a/Engine/Plugins/Experimental/CinematicPrestreaming/Source/CinematicPrestreamingEditor/Private/CinePrestreamingDebugRender.h b/Engine/Plugins/Experimental/CinematicPrestreaming/Source/CinematicPrestreamingEditor/Private/CinePrestreamingDebugRender.h
--- a/Engine/Plugins/Experimental/CinematicPrestreaming/Source/CinematicPrestreamingEditor/Private/CinePrestreamingDebugRender.h
+++ b/Engine/Plugins/Experimental/CinematicPrestreaming/Source/CinematicPrestreamingEditor/Private/CinePrestreamingDebugRender.h
@@ -5,24 +5,27 @@
 #include "MoviePipelineDeferredPasses.h"
 #include "CinePrestreamingDebugRender.generated.h"
 
+enum class EVirtualTextureVisualizationMode : uint8;
+
 UCLASS(BlueprintType)
 class UCinePrestreamingDebugRender : public UMoviePipelineDeferredPassBase
 {
 	GENERATED_BODY()
 
 public:
-	UCinePrestreamingDebugRender() : UMoviePipelineDeferredPassBase()
-	{
-		PassIdentifier = FMoviePipelinePassIdentifier("VirtualTexturePendingMips");
-	}
+	UCinePrestreamingDebugRender();
+
+	virtual void SetupImpl(const MoviePipeline::FMoviePipelineRenderPassInitSettings& InPassInitSettings) override;
+	virtual void TeardownImpl() override;
 #if WITH_EDITOR
-	virtual FText GetDisplayText() const override { return NSLOCTEXT("MovieRenderPipeline", "DeferredBasePassSetting_DisplayName_VTPendingMips", "VT Pending Mips (Debug)"); }
+	virtual FText GetDisplayText() const override;
 #endif
-	virtual void GetViewShowFlags(FEngineShowFlags& OutShowFlag, EViewModeIndex& OutViewModeIndex) const override
-	{
-		OutShowFlag = FEngineShowFlags(EShowFlagInitMode::ESFIM_Game);
-		OutShowFlag.SetVisualizeVirtualTexture(true);
-		OutViewModeIndex = EViewModeIndex::VMI_VisualizeVirtualTexture;
-	}
+	virtual void GetViewShowFlags(FEngineShowFlags& OutShowFlag, EViewModeIndex& OutViewModeIndex) const override;
 	virtual int32 GetOutputFileSortingOrder() const override { return 2; }
+
+ 	UPROPERTY(EditAnywhere, BlueprintReadWrite, Category = "Settings")
+ 	EVirtualTextureVisualizationMode VirtualTextureDebugMode;
+
+private:
+	EVirtualTextureVisualizationMode PreviousMode;
 };
