Cinematic Prestream: Change default behavior for loading of prestreaming data so that we use hard refs instead of async loads.
This ensures that the data is always available in time for playback without manual tweaking of preroll timings.
The downside is that the prestreaming shot data for a cinematic needs to stay resident for the length of the cinematic, which increases memory ceiling.
The support for async load remains, but it is now a (non-default) option.
#Synced-CL 50433575

diff --git a/Engine/Plugins/Experimental/CinematicPrestreaming/Source/CinematicPrestreaming/Private/CinePrestreamingData.cpp b/Engine/Plugins/Experimental/CinematicPrestreaming/Source/CinematicPrestreaming/Private/CinePrestreamingData.cpp
--- a/Engine/Plugins/Experimental/CinematicPrestreaming/Source/CinematicPrestreaming/Private/CinePrestreamingData.cpp
+++ b/Engine/Plugins/Experimental/CinematicPrestreaming/Source/CinematicPrestreaming/Private/CinePrestreamingData.cpp
@@ -20,5 +20,7 @@ void UCinePrestreamingData::GetResourceSizeEx(FResourceSizeEx& CumulativeResourc
 	{
 		Size += Data.RequestData.GetAllocatedSize();
 	}
-	CumulativeResourceSize.AddDedicatedSystemMemoryBytes(Size);
+	
+	static FName TagName(TEXT("CinePrestreaming Data"));
+	CumulativeResourceSize.AddDedicatedSystemMemoryBytes(TagName, Size);
 }
diff --git a/Engine/Plugins/Experimental/CinematicPrestreaming/Source/CinematicPrestreaming/Private/Sections/MovieSceneCinePrestreamingSection.cpp b/Engine/Plugins/Experimental/CinematicPrestreaming/Source/CinematicPrestreaming/Private/Sections/MovieSceneCinePrestreamingSection.cpp
--- a/Engine/Plugins/Experimental/CinematicPrestreaming/Source/CinematicPrestreaming/Private/Sections/MovieSceneCinePrestreamingSection.cpp
+++ b/Engine/Plugins/Experimental/CinematicPrestreaming/Source/CinematicPrestreaming/Private/Sections/MovieSceneCinePrestreamingSection.cpp
@@ -1,9 +1,12 @@
 // Copyright Epic Games, Inc. All Rights Reserved.
 
 #include "Sections/MovieSceneCinePrestreamingSection.h"
+
+#include "Engine/AssetManager.h"
+#include "Engine/StreamableManager.h"
 #include "EntitySystem/BuiltInComponentTypes.h"
-#include "TrackInstances/MovieSceneCinePrestreamingTrackInstance.h"
 #include "EntitySystem/MovieSceneEntityBuilder.h"
+#include "TrackInstances/MovieSceneCinePrestreamingTrackInstance.h"
 
 UMovieSceneCinePrestreamingSection::UMovieSceneCinePrestreamingSection(const FObjectInitializer& ObjInit)
 	: Super(ObjInit)
@@ -11,6 +14,31 @@ UMovieSceneCinePrestreamingSection::UMovieSceneCinePrestreamingSection(const FOb
 	EvalOptions.EnableAndSetCompletionMode(EMovieSceneCompletionMode::RestoreState);
 }
 
+void UMovieSceneCinePrestreamingSection::Serialize(FArchive& Ar)
+{
+#if WITH_EDITOR	
+	// If we want async load behavior then cook the soft instead of the hard ref.
+	if (Ar.IsCooking() && Ar.IsSaving() && bUseAsyncLoading)
+	{
+		TObjectPtr<UCinePrestreamingData> PrestreamingAssetBackup = PrestreamingAsset;
+		// Note that we always expect SoftPrestreamingAsset to be null in editor, but we do backup/restore here for completeness.
+		TSoftObjectPtr<UCinePrestreamingData> SoftPrestreamingAssetBackup = SoftPrestreamingAsset;
+		
+		SoftPrestreamingAsset = PrestreamingAsset;
+		PrestreamingAsset = nullptr;
+
+		Super::Serialize(Ar);
+
+		PrestreamingAsset = PrestreamingAssetBackup;
+		SoftPrestreamingAsset = SoftPrestreamingAssetBackup;
+	}
+	else
+#endif
+	{
+		Super::Serialize(Ar);
+	}
+}
+
 void UMovieSceneCinePrestreamingSection::ImportEntityImpl(UMovieSceneEntitySystemLinker* EntityLinker, const FEntityImportParams& Params, FImportedEntity* OutImportedEntity)
 {
 	using namespace UE::MovieScene;
@@ -23,3 +51,46 @@ void UMovieSceneCinePrestreamingSection::ImportEntityImpl(UMovieSceneEntitySyste
 		.Add(FBuiltInComponentTypes::Get()->TrackInstance, TrackInstance)
 	);
 }
+
+void UMovieSceneCinePrestreamingSection::AsyncStreamingAddRef()
+{
+	if (!bUseAsyncLoading || SoftPrestreamingAsset.IsNull())
+	{
+		return;
+	}
+	
+	if (AsyncLoadingRefCount++ > 0)
+	{
+		return;
+	}
+
+	FStreamableManager& StreamableManager = UAssetManager::GetStreamableManager();
+	TWeakObjectPtr<UMovieSceneCinePrestreamingSection> WeakThis = this;
+	AsyncLoadHandle = StreamableManager.RequestAsyncLoad(SoftPrestreamingAsset.ToSoftObjectPath(), [WeakThis]()
+	{
+		if (UMovieSceneCinePrestreamingSection* HardThis = WeakThis.Get())
+		{
+			HardThis->PrestreamingAsset = HardThis->SoftPrestreamingAsset.Get();
+		}
+	});
+}
+
+void UMovieSceneCinePrestreamingSection::AsyncStreamingReleaseRef()
+{
+	if (!bUseAsyncLoading || SoftPrestreamingAsset.IsNull())
+	{
+		return;
+	}
+	
+	if (AsyncLoadingRefCount == 0 || --AsyncLoadingRefCount > 0)
+	{
+		return;
+	}
+
+	if (AsyncLoadHandle.IsValid())
+	{
+		AsyncLoadHandle->CancelHandle();
+		AsyncLoadHandle.Reset();
+	}
+	PrestreamingAsset = nullptr;
+}
diff --git a/Engine/Plugins/Experimental/CinematicPrestreaming/Source/CinematicPrestreaming/Private/TrackInstances/MovieSceneCinePrestreamingTrackInstance.cpp b/Engine/Plugins/Experimental/CinematicPrestreaming/Source/CinematicPrestreaming/Private/TrackInstances/MovieSceneCinePrestreamingTrackInstance.cpp
--- a/Engine/Plugins/Experimental/CinematicPrestreaming/Source/CinematicPrestreaming/Private/TrackInstances/MovieSceneCinePrestreamingTrackInstance.cpp
+++ b/Engine/Plugins/Experimental/CinematicPrestreaming/Source/CinematicPrestreaming/Private/TrackInstances/MovieSceneCinePrestreamingTrackInstance.cpp
@@ -2,11 +2,8 @@
 
 #include "TrackInstances/MovieSceneCinePrestreamingTrackInstance.h"
 
-#include "CinePrestreamingLog.h"
-#include "Engine/AssetManager.h"
-#include "Engine/StreamableManager.h"
+#include "CinePrestreamingData.h"
 #include "EngineModule.h"
-#include "EntitySystem/MovieSceneInstanceRegistry.h"
 #include "MovieScene.h"
 #include "RendererInterface.h"
 #include "RenderingThread.h"
@@ -21,64 +18,31 @@ static TAutoConsoleVariable<int32> CVarCinematicPreStreamQualityLevel(
 	ECVF_ReadOnly
 );
 
+static bool SupportSectionQualityLevel(UMovieSceneCinePrestreamingSection const* InSection)
+{
+	return CVarCinematicPreStreamQualityLevel.GetValueOnGameThread() >= InSection->GetQualityLevel();
+}
+
 void UMovieSceneCinePrestreamingTrackInstance::OnInputAdded(const FMovieSceneTrackInstanceInput& InInput)
 {
-	// Trigger async load for the asset associated with this input.
-	const UMovieSceneCinePrestreamingSection* Section = Cast<const UMovieSceneCinePrestreamingSection>(InInput.Section);
-	if (Section != nullptr)
+	if (UMovieSceneCinePrestreamingSection* Section = Cast<UMovieSceneCinePrestreamingSection>(InInput.Section))
 	{
-		const int32 QualityLevel = Section->GetQualityLevel();
-		if (CVarCinematicPreStreamQualityLevel.GetValueOnGameThread() < QualityLevel)
+		if (SupportSectionQualityLevel(Section))
 		{
-			return;
-		}
-
-		TSoftObjectPtr<UCinePrestreamingData> AssetSoftPtr = Section->GetPrestreamingAsset();
-		if (!AssetSoftPtr.IsNull())
-		{
-			if (PrestreamingAssetMap.Find(InInput) == nullptr)
-			{
-				TWeakObjectPtr<UMovieSceneCinePrestreamingTrackInstance> WeakThis = this;
-				FMovieSceneTrackInstanceInput Input = InInput;
-
-				FStreamableDelegate OnLoadedDelegate;
-				OnLoadedDelegate.BindLambda(
-					[WeakThis, Input, AssetSoftPtr]()
-					{
-						if (UMovieSceneCinePrestreamingTrackInstance* StrongThis = WeakThis.Get())
-						{
-							if (UCinePrestreamingData* AssetPtr = AssetSoftPtr.Get())
-							{
-								StrongThis->PrestreamingAssetMap.Add(Input, AssetPtr);
-								StrongThis->LoadHandleMap.Remove(Input);
-							}
-						}
-					});
-		
-				FStreamableManager& StreamableManager = UAssetManager::GetStreamableManager();
-				TSharedPtr<FStreamableHandle> LoadHandle = StreamableManager.RequestAsyncLoad(AssetSoftPtr.ToSoftObjectPath(), OnLoadedDelegate);
-				LoadHandleMap.Add(InInput, LoadHandle);
-			}
+			// Trigger any required async load for the prestreaming data.
+			Section->AsyncStreamingAddRef();
 		}
 	}
 }
 
 void UMovieSceneCinePrestreamingTrackInstance::OnInputRemoved(const FMovieSceneTrackInstanceInput& InInput)
 {
-	const UMovieSceneCinePrestreamingSection* Section = Cast<const UMovieSceneCinePrestreamingSection>(InInput.Section);
-	if (Section != nullptr)
+	if (UMovieSceneCinePrestreamingSection* Section = Cast<UMovieSceneCinePrestreamingSection>(InInput.Section))
 	{
-		// Remove asset reference should lead to asset unload at some future time.
-		if (PrestreamingAssetMap.Remove(InInput) == 0)
+		if (SupportSectionQualityLevel(Section))
 		{
-			// Didn't remove, so maybe we still have an outstanding load handle.
-			TSharedPtr<FStreamableHandle>* LoadHandlePtr = LoadHandleMap.Find(InInput);
-			if (LoadHandlePtr != nullptr)
-			{
-				FStreamableManager& StreamableManager = UAssetManager::GetStreamableManager();
-				(*LoadHandlePtr)->CancelHandle();
-				LoadHandleMap.Remove(InInput);
-			}
+			// Release any async load for the prestreaming data.
+			Section->AsyncStreamingReleaseRef();
 		}
 	}
 }
@@ -97,6 +61,18 @@ void UMovieSceneCinePrestreamingTrackInstance::OnAnimate()
 
 	for (const FMovieSceneTrackInstanceInput& Input : GetInputs())
 	{
+		UMovieSceneCinePrestreamingSection* Section = Cast<UMovieSceneCinePrestreamingSection>(Input.Section);
+		if (Section == nullptr)
+		{
+			continue;
+		}
+
+		// Early out if this input is stripped by QualityLevel setting.
+		if (!SupportSectionQualityLevel(Section))
+		{
+			continue;
+		}
+
 		const FSequenceInstance& SequenceInstance = InstanceRegistry->GetInstance(Input.InstanceHandle);
 		const FMovieSceneContext& Context = SequenceInstance.GetContext();
 
@@ -123,26 +99,16 @@ void UMovieSceneCinePrestreamingTrackInstance::OnAnimate()
 			continue;
 		}
 
-		// Early out if this input is stripped by QualityLevel setting.
-		const UMovieSceneCinePrestreamingSection* Section = Cast<const UMovieSceneCinePrestreamingSection>(Input.Section);
-		const int32 QualityLevel = Section->GetQualityLevel();
-		if (CVarCinematicPreStreamQualityLevel.GetValueOnGameThread() < QualityLevel)
+		UCinePrestreamingData* PrestreamingData = Section->GetPrestreamingAsset();
+		if (PrestreamingData == nullptr)
 		{
+			// We can get here if there is no prestreaming data asset set or if async loading has not had time to complete.
 			continue;
 		}
 
-		// Get (async loaded) asset data.
-		TObjectPtr<UCinePrestreamingData>* PrestreamingDataPtr = PrestreamingAssetMap.Find(Input);
-		
-		// PreRoll needs to be long enough to cover any async loading before evaluation starts.
-		// In editor this can happen if we are scrubbing the track. But that case is fine.
-		if (!Section->GetPrestreamingAsset().IsNull() && PrestreamingDataPtr == nullptr)
-		{
-			UE_LOG(LogCinePrestreaming, Verbose, TEXT("Prestreaming asset %s didn't load in time. Consider using a longer PreRoll."), *Section->GetPrestreamingAsset().GetAssetName());
-		}
-
-		UCinePrestreamingData const* PrestreamingData = PrestreamingDataPtr != nullptr ? PrestreamingDataPtr->Get() : nullptr;
-		if (PrestreamingData == nullptr || PrestreamingData->Times.Num() == 0)
+		// Early out if there is no prestreaming data content, or if data is invalid.
+		const int32 ItemCount = PrestreamingData->Times.Num();
+		if (ItemCount == 0 || ItemCount != PrestreamingData->VirtualTextureDatas.Num() || ItemCount != PrestreamingData->NaniteDatas.Num())
 		{
 			continue;
 		}
diff --git a/Engine/Plugins/Experimental/CinematicPrestreaming/Source/CinematicPrestreaming/Private/TrackInstances/MovieSceneCinePrestreamingTrackInstance.h b/Engine/Plugins/Experimental/CinematicPrestreaming/Source/CinematicPrestreaming/Private/TrackInstances/MovieSceneCinePrestreamingTrackInstance.h
--- a/Engine/Plugins/Experimental/CinematicPrestreaming/Source/CinematicPrestreaming/Private/TrackInstances/MovieSceneCinePrestreamingTrackInstance.h
+++ b/Engine/Plugins/Experimental/CinematicPrestreaming/Source/CinematicPrestreaming/Private/TrackInstances/MovieSceneCinePrestreamingTrackInstance.h
@@ -5,9 +5,6 @@
 #include "EntitySystem/TrackInstance/MovieSceneTrackInstance.h"
 #include "MovieSceneCinePrestreamingTrackInstance.generated.h"
 
-struct FStreamableHandle;
-class UCinePrestreamingData;
-
 UCLASS()
 class UMovieSceneCinePrestreamingTrackInstance : public UMovieSceneTrackInstance
 {
@@ -18,12 +15,4 @@ protected:
 	void OnInputAdded(const FMovieSceneTrackInstanceInput& InInput) override;
 	void OnInputRemoved(const FMovieSceneTrackInstanceInput& InInput) override;
 	void OnAnimate() override;
-
-private:
-	/** Map of loaded asset references. */
-	UPROPERTY(Transient)
-	TMap< FMovieSceneTrackInstanceInput, TObjectPtr<UCinePrestreamingData> > PrestreamingAssetMap;
-
-	/** Map of outstanding async loading handles. */
-	TMap< FMovieSceneTrackInstanceInput, TSharedPtr<FStreamableHandle> > LoadHandleMap;
 };
diff --git a/Engine/Plugins/Experimental/CinematicPrestreaming/Source/CinematicPrestreaming/Public/Sections/MovieSceneCinePrestreamingSection.h b/Engine/Plugins/Experimental/CinematicPrestreaming/Source/CinematicPrestreaming/Public/Sections/MovieSceneCinePrestreamingSection.h
--- a/Engine/Plugins/Experimental/CinematicPrestreaming/Source/CinematicPrestreaming/Public/Sections/MovieSceneCinePrestreamingSection.h
+++ b/Engine/Plugins/Experimental/CinematicPrestreaming/Source/CinematicPrestreaming/Public/Sections/MovieSceneCinePrestreamingSection.h
@@ -8,6 +8,7 @@
 #include "MovieSceneCinePrestreamingSection.generated.h"
 
 class UCinePrestreamingData;
+struct FStreamableHandle;
 
 /** Movie Scene Section representing a Prestreaming asset. */
 UCLASS(MinimalAPI)
@@ -20,12 +21,12 @@ class UMovieSceneCinePrestreamingSection
 public:
 	UMovieSceneCinePrestreamingSection(const FObjectInitializer& ObjInit);
 
-	/** Get the prestreaming asset soft pointer. Use for async loading. */
+	/** Get the prestreaming asset pointer. This may return nullptr in cooked build if we are using async loading. */
 	UFUNCTION(BlueprintPure, Category = "Cinematic Prestreaming")
-	TSoftObjectPtr<UCinePrestreamingData> GetPrestreamingAsset() const { return PrestreamingAsset; }
+	UCinePrestreamingData* GetPrestreamingAsset() const { return PrestreamingAsset; }
 
 	UFUNCTION(BlueprintCallable, Category = "Cinematic Prestreaming")
-	void SetPrestreamingAsset(const UCinePrestreamingData* InData) { PrestreamingAsset = TSoftObjectPtr<UCinePrestreamingData>(const_cast<UCinePrestreamingData*>(InData)); }
+	void SetPrestreamingAsset(UCinePrestreamingData* InData) { PrestreamingAsset = InData; }
 
 	/** If MovieScene.PreStream.QualityLevel is less than this then discard this section at runtime. */
 	UFUNCTION(BlueprintPure, Category = "Cinematic Prestreaming")
@@ -37,18 +38,30 @@ public:
 	UFUNCTION(BlueprintCallable, Category = "Cinematic Prestreaming")
 	void SetStartFrameOffset(const int32 InOffset) { StartFrameOffset = InOffset; }
 
+	/** Add a reference to the async loader. When the ref count is non-zero the prestreaming asset should be loaded or loading. */
+	void AsyncStreamingAddRef();
+
+	/** Release the reference to the async loader. When the ref count is zero the prestreaming asset can be unloaded. */
+	void AsyncStreamingReleaseRef();
+
 	/** UMovieSceneSection interface */
-	TOptional<FFrameTime> GetOffsetTime() const override { return TOptional<FFrameTime>(StartFrameOffset); }
-	
+	virtual TOptional<FFrameTime> GetOffsetTime() const override { return TOptional<FFrameTime>(StartFrameOffset); }
 
 protected:
+	/** UObject interface */
+	virtual void Serialize(FArchive& Ar) override;
+
 	/** IMovieSceneEntityProvider interface */
-	void ImportEntityImpl(UMovieSceneEntitySystemLinker* EntityLinker, const FEntityImportParams& Params, FImportedEntity* OutImportedEntity) override;
+	virtual void ImportEntityImpl(UMovieSceneEntitySystemLinker* EntityLinker, const FEntityImportParams& Params, FImportedEntity* OutImportedEntity) override;
 
 protected:
 	/** The asset containing cinematic prestreaming data to use for this section. */
 	UPROPERTY(EditAnywhere, BlueprintReadWrite, Category = Prestreaming)
-	TSoftObjectPtr<UCinePrestreamingData> PrestreamingAsset;
+	TObjectPtr<UCinePrestreamingData> PrestreamingAsset;
+
+	/** Soft ref version of the asset for use by async loading. */
+	UPROPERTY()
+	TSoftObjectPtr<UCinePrestreamingData> SoftPrestreamingAsset;
 
 	/** Number of frames by which to offset the evaluation. Larger values cause prestreaming to happen earlier. */
 	UPROPERTY(EditAnywhere, BlueprintReadWrite, Category = Prestreaming)
@@ -57,15 +70,18 @@ protected:
 	/** If If MovieScene.PreStream.QualityLevel is less than this then discard this section at runtime. */
 	UPROPERTY(EditAnywhere, BlueprintReadWrite, Category = Prestreaming, AdvancedDisplay)
 	int32 QualityLevel = 0;
-};
 
+	/** 
+	 * Set true to enable async loading of the prestreaming data in game.
+	 * This can reduce memory overhead for a long cinematic with multiple prestreaming shots.
+	 * Async loading requires sufficient preroll time to load the asset and still have enough time to run StartFrameOffset frames ahead.
+	 */
+	UPROPERTY(EditAnywhere, BlueprintReadWrite, Category = Prestreaming, AdvancedDisplay)
+	bool bUseAsyncLoading= false;
 
-/** Component data for the cinematic prestreaming system. */
-USTRUCT()
-struct FMovieSceneCinePrestreamingComponentData
-{
-	GENERATED_BODY()
+	/** Ref count for async loading of PrestreamingAsset. */
+	uint32 AsyncLoadingRefCount = 0;
 
-	UPROPERTY()
-	TObjectPtr<const UMovieSceneCinePrestreamingSection> Section = nullptr;
+	/** Handle for async loading of PrestreamingAsset. */
+	TSharedPtr<FStreamableHandle> AsyncLoadHandle;
 };
diff --git a/Engine/Plugins/Experimental/CinematicPrestreaming/Source/CinematicPrestreamingEditor/Private/TrackEditors/CinePrestreamingTrackEditor.cpp b/Engine/Plugins/Experimental/CinematicPrestreaming/Source/CinematicPrestreamingEditor/Private/TrackEditors/CinePrestreamingTrackEditor.cpp
--- a/Engine/Plugins/Experimental/CinematicPrestreaming/Source/CinematicPrestreamingEditor/Private/TrackEditors/CinePrestreamingTrackEditor.cpp
+++ b/Engine/Plugins/Experimental/CinematicPrestreaming/Source/CinematicPrestreamingEditor/Private/TrackEditors/CinePrestreamingTrackEditor.cpp
@@ -82,9 +82,9 @@ struct FCinePrestreamingSection
 	FText GetVisibilityText() const
 	{
 		UMovieSceneCinePrestreamingSection* Section = WeakSection.Get();
-		if (Section && !Section->GetPrestreamingAsset().IsNull())
+		if (Section && Section->GetPrestreamingAsset())
 		{
-			return FText::FromString(FString::Printf(TEXT("Asset: %s"), *Section->GetPrestreamingAsset().GetAssetName()));
+			return FText::FromString(FString::Printf(TEXT("Asset: %s"), *Section->GetPrestreamingAsset()->GetName()));
 		}
 
 		return LOCTEXT("VisibilityText_NoAsset", "No Asset Specified");
@@ -107,7 +107,7 @@ struct FCinePrestreamingSection
 		UMovieSceneCinePrestreamingSection* Section = WeakSection.Get();
 		if (Section)
 		{
-			if (const UCinePrestreamingData* PrestreamingData = Section->GetPrestreamingAsset().Get())
+			if (const UCinePrestreamingData* PrestreamingData = Section->GetPrestreamingAsset())
 			{
 				FDateTime CurTime = FDateTime::UtcNow();
 				FTimespan DeltaSinceRecording = (CurTime - PrestreamingData->RecordedTime);
