Set UCinePrestreamingData as BlueprintType.
That is required to be able to use UMovieSceneCinePrestreamingSection::SetPrestreamingAsset() from Blueprint.
#rnx
#Synced-CL 51290350

diff --git a/Engine/Plugins/Experimental/CinematicPrestreaming/Source/CinematicPrestreaming/Public/CinePrestreamingData.h b/Engine/Plugins/Experimental/CinematicPrestreaming/Source/CinematicPrestreaming/Public/CinePrestreamingData.h
--- a/Engine/Plugins/Experimental/CinematicPrestreaming/Source/CinematicPrestreaming/Public/CinePrestreamingData.h
+++ b/Engine/Plugins/Experimental/CinematicPrestreaming/Source/CinematicPrestreaming/Public/CinePrestreamingData.h
@@ -33,7 +33,7 @@ struct FCinePrestreamingNaniteData
 	TArray<uint32> RequestData;
 };
 
-UCLASS()
+UCLASS(BlueprintType)
 class CINEMATICPRESTREAMING_API UCinePrestreamingData : public UObject
 {
 	GENERATED_BODY()
